[WebSocketMessaging] Moved plugin

diff --git a/Engine/Plugins/Experimental/WebSocketMessaging/Source/WebSocketMessaging/Private/WebSocketDeserializedMessage.cpp b/Engine/Plugins/Experimental/WebSocketMessaging/Source/WebSocketMessaging/Private/WebSocketDeserializedMessage.cpp
new file mode 100644
--- /dev/null
+++ b/Engine/Plugins/Experimental/WebSocketMessaging/Source/WebSocketMessaging/Private/WebSocketDeserializedMessage.cpp
@@ -0,0 +1,141 @@
+// Copyright Epic Games, Inc. All Rights Reserved.
+
+
+#include "WebSocketDeserializedMessage.h"
+#include "JsonObjectConverter.h"
+#include "WebSocketMessaging.h"
+
+FWebSocketDeserializedMessage::FWebSocketDeserializedMessage()
+	: Expiration(FDateTime::MaxValue())	// Make sure messages don't expire if no expiration is specified.
+	, Message(nullptr)
+	, Scope(EMessageScope::All)
+{
+}
+
+FWebSocketDeserializedMessage::~FWebSocketDeserializedMessage()
+{
+	if (Message)
+	{
+		FMemory::Free(Message);
+	}
+}
+
+bool FWebSocketDeserializedMessage::ParseJson(const FString& Json)
+{
+	static TMap<FString, EMessageScope> MessageScopeStringMapping =
+	{
+		{"Thread", EMessageScope::Thread},
+		{ "Process", EMessageScope::Process },
+		{ "Network", EMessageScope::Network },
+		{ "All", EMessageScope::All }
+	};
+
+	TSharedPtr<FJsonValue> RootValue;
+
+	TSharedRef<TJsonReader<TCHAR>> JsonReader = TJsonReaderFactory<TCHAR>::Create(Json);
+	if (FJsonSerializer::Deserialize(JsonReader, RootValue))
+	{
+		TSharedPtr<FJsonObject> RootObject = RootValue->AsObject();
+		if (!RootObject)
+		{
+			return false;
+		}
+
+		FString MessageType;
+		if (!RootObject->TryGetStringField(WebSocketMessaging::Tag::MessageType, MessageType))
+		{
+			return false;
+		}
+
+		FString JsonSender;
+		if (!RootObject->TryGetStringField(WebSocketMessaging::Tag::Sender, JsonSender))
+		{
+			return false;
+		}
+
+		if (!FMessageAddress::Parse(JsonSender, Sender))
+		{
+			return false;
+		}
+
+		const TSharedPtr<FJsonObject>* JsonAnnotations = nullptr;
+		if (RootObject->TryGetObjectField(WebSocketMessaging::Tag::Annotations, JsonAnnotations))
+		{
+			for (const TPair<FString, TSharedPtr<FJsonValue>>& Pair : (*JsonAnnotations)->Values)
+			{
+				Annotations.Add(*Pair.Key, Pair.Value->AsString());
+			}
+		}
+
+		const TSharedPtr<FJsonObject>* JsonMessage = nullptr;
+		if (!RootObject->TryGetObjectField(WebSocketMessaging::Tag::Message, JsonMessage))
+		{
+			return false;
+		}
+
+		UScriptStruct* ScriptStruct = FindObjectSafe<UScriptStruct>(nullptr, *MessageType);
+		if (!ScriptStruct)
+		{
+			return false;
+		}
+
+		TypeInfo = ScriptStruct;
+
+		if (Message)
+		{
+			FMemory::Free(Message);
+		}
+
+		Message = FMemory::Malloc(ScriptStruct->GetStructureSize());
+		ScriptStruct->InitializeStruct(Message);
+
+		if (!FJsonObjectConverter::JsonObjectToUStruct(JsonMessage->ToSharedRef(), ScriptStruct, Message))
+		{
+			FMemory::Free(Message);
+			Message = nullptr;
+			return false;
+		}
+
+		int64 UnixTime;
+		if (RootObject->TryGetNumberField(WebSocketMessaging::Tag::Expiration, UnixTime))
+		{
+			Expiration = FDateTime::FromUnixTimestamp(UnixTime);
+		}
+
+		if (RootObject->TryGetNumberField(WebSocketMessaging::Tag::TimeSent, UnixTime))
+		{
+			TimeSent = FDateTime::FromUnixTimestamp(UnixTime);
+		}
+
+		FString ScopeString;
+		if (RootObject->TryGetStringField(WebSocketMessaging::Tag::Scope, ScopeString))
+		{
+			if (MessageScopeStringMapping.Contains(ScopeString))
+			{
+				Scope = MessageScopeStringMapping[ScopeString];
+			}
+			// unknown scope string
+			else
+			{
+				return false;
+			}
+		}
+
+		TArray<FString> RecipientsStrings;
+		if (RootObject->TryGetStringArrayField(WebSocketMessaging::Tag::Recipients, RecipientsStrings))
+		{
+			for (const FString& RecipientString : RecipientsStrings)
+			{
+				FMessageAddress RecipientAddress;
+				if (FMessageAddress::Parse(RecipientString, RecipientAddress))
+				{
+					Recipients.Add(RecipientAddress);
+				}
+			}
+		}
+
+		return true;
+	}
+
+	return false;
+}
diff --git a/Engine/Plugins/Experimental/WebSocketMessaging/Source/WebSocketMessaging/Private/WebSocketMessageTransport.cpp b/Engine/Plugins/Experimental/WebSocketMessaging/Source/WebSocketMessaging/Private/WebSocketMessageTransport.cpp
new file mode 100644
--- /dev/null
+++ b/Engine/Plugins/Experimental/WebSocketMessaging/Source/WebSocketMessaging/Private/WebSocketMessageTransport.cpp
@@ -0,0 +1,420 @@
+// Copyright Epic Games, Inc. All Rights Reserved.
+
+
+#include "WebSocketMessageTransport.h"
+#include "Backends/CborStructSerializerBackend.h"
+#include "CborReader.h"
+#include "CborWriter.h"
+#include "IMessageContext.h"
+#include "IMessageTransportHandler.h"
+#include "INetworkingWebSocket.h"
+#include "IWebSocket.h"
+#include "IWebSocketNetworkingModule.h"
+#include "IWebSocketServer.h"
+#include "JsonObjectConverter.h"
+#include "Serialization/ArrayReader.h"
+#include "Serialization/ArrayWriter.h"
+#include "StructSerializer.h"
+#include "WebSocketDeserializedMessage.h"
+#include "WebSocketMessaging.h"
+#include "WebSocketMessagingSettings.h"
+#include "WebSocketsModule.h"
+
+bool FWebSocketMessageConnection::IsConnected() const
+{
+	if (WebSocketConnection)
+	{
+		return WebSocketConnection->IsConnected();
+	}
+
+	return WebSocketServerConnection != nullptr;
+}
+
+void FWebSocketMessageConnection::Close()
+{
+	if (WebSocketConnection)
+	{
+		return WebSocketConnection->Close();
+	}
+}
+
+FWebSocketMessageTransport::FWebSocketMessageTransport()
+{
+}
+
+FWebSocketMessageTransport::~FWebSocketMessageTransport()
+{
+}
+
+bool FWebSocketMessageTransport::StartTransport(IMessageTransportHandler& Handler)
+{
+	const UWebSocketMessagingSettings* Settings = GetDefault<UWebSocketMessagingSettings>();
+
+	TransportHandler = &Handler;
+
+	if (Settings->ServerPort > 0)
+	{
+		IWebSocketNetworkingModule* WebSocketNetworkingModule = FModuleManager::Get().LoadModulePtr<IWebSocketNetworkingModule>(TEXT("WebSocketNetworking"));
+		if (WebSocketNetworkingModule)
+		{
+			Server = WebSocketNetworkingModule->CreateServer();
+			if (Server)
+			{
+				if (!Server->Init(Settings->ServerPort, FWebSocketClientConnectedCallBack::CreateThreadSafeSP(this, &FWebSocketMessageTransport::ClientConnected)))
+				{
+					UE_LOG(LogWebSocketMessaging, Log, TEXT("Unable to start WebSocketMessaging Server on port %d"), Settings->ServerPort);
+				}
+				else
+				{
+					ServerTickerHandle = FTSTicker::GetCoreTicker().AddTicker(FTickerDelegate::CreateThreadSafeSP(this, &FWebSocketMessageTransport::ServerTick));
+					UE_LOG(LogWebSocketMessaging, Log, TEXT("WebSocketMessaging Server started on port %d"), Settings->ServerPort);
+				}
+			}
+		}
+		else
+		{
+			UE_LOG(LogWebSocketMessaging, Log, TEXT("Unable to load WebSocketNetworking module, ensure to enable it"));
+		}
+	}
+
+	for (const FString& Url : Settings->ConnectToEndpoints)
+	{
+		FGuid Guid = FGuid::NewGuid();
+
+		TMap<FString, FString> Headers;
+		Headers.Add(WebSocketMessaging::Header::TransportId, Guid.ToString());
+		for (const TPair<FString, FString>& Pair : Settings->HttpHeaders)
+		{
+			Headers.Add(Pair.Key, Pair.Value);
+		}
+
+		TSharedRef<IWebSocket, ESPMode::ThreadSafe> WebSocketConnection = FWebSocketsModule::Get().CreateWebSocket(Url, FString(), Headers);
+
+		FWebSocketMessageConnectionRef WebSocketMessageConnection = MakeShared<FWebSocketMessageConnection>(Url, Guid, WebSocketConnection);
+
+		WebSocketConnection->OnMessage().AddThreadSafeSP(this, &FWebSocketMessageTransport::OnJsonMessage, WebSocketMessageConnection);
+		WebSocketConnection->OnClosed().AddThreadSafeSP(this, &FWebSocketMessageTransport::OnClosed, WebSocketMessageConnection);
+		WebSocketConnection->OnConnected().AddThreadSafeSP(this, &FWebSocketMessageTransport::OnConnected, WebSocketMessageConnection);
+		WebSocketConnection->OnConnectionError().AddThreadSafeSP(this, &FWebSocketMessageTransport::OnConnectionError, WebSocketMessageConnection);
+
+		WebSocketMessageConnections.Add(Guid, WebSocketMessageConnection);
+
+		WebSocketConnection->Connect();
+
+	}
+
+	return true;
+}
+
+void FWebSocketMessageTransport::StopTransport()
+{
+	if (Server.IsValid())
+	{
+		Server.Reset();
+	}
+
+	for (TPair<FGuid, FWebSocketMessageConnectionRef> Pair : WebSocketMessageConnections)
+	{
+		Pair.Value->bDestroyed = true;
+		Pair.Value->Close();
+	}
+	WebSocketMessageConnections.Empty();
+}
+
+void FWebSocketMessageTransport::OnClosed(int32 Code, const FString& Reason, bool bUserClose, FWebSocketMessageConnectionRef WebSocketMessageConnection)
+{
+	UE_LOG(LogWebSocketMessaging, Log, TEXT("Connection to %s closed, Code: %d Reason: \"%s\" UserClose: %s, retrying..."), *WebSocketMessageConnection->Url, Code, *Reason, bUserClose ? TEXT("true") : TEXT("false"));
+	ForgetTransportNode(WebSocketMessageConnection);
+	WebSocketMessageConnection->bIsConnecting = false;
+	RetryConnection(WebSocketMessageConnection);
+}
+
+void FWebSocketMessageTransport::OnConnectionError(const FString& Message, FWebSocketMessageConnectionRef WebSocketMessageConnection)
+{
+	if (!WebSocketMessageConnection->bIsConnecting)
+	{
+		UE_LOG(LogWebSocketMessaging, Log, TEXT("Connection to %s error: %s, retrying..."), *WebSocketMessageConnection->Url, *Message);
+	}
+	ForgetTransportNode(WebSocketMessageConnection);
+	WebSocketMessageConnection->bIsConnecting = false;
+	RetryConnection(WebSocketMessageConnection);
+}
+
+void FWebSocketMessageTransport::OnJsonMessage(const FString& Message, FWebSocketMessageConnectionRef WebSocketMessageConnection)
+{
+	TSharedRef<FWebSocketDeserializedMessage> Context = MakeShared<FWebSocketDeserializedMessage>();
+	if (Context->ParseJson(Message))
+	{
+		TransportHandler->ReceiveTransportMessage(Context, WebSocketMessageConnection->Guid);
+	}
+	else
+	{
+		UE_LOG(LogWebSocketMessaging, Verbose, TEXT("Invalid Json Message received on %s"), *WebSocketMessageConnection->Url);
+	}
+}
+
+void FWebSocketMessageTransport::OnServerJsonMessage(void* Data, int32 DataSize, FWebSocketMessageConnectionRef WebSocketMessageConnection)
+{
+	FString Message(DataSize, reinterpret_cast<UTF8CHAR*>(Data));
+	OnJsonMessage(Message, WebSocketMessageConnection);
+}
+
+class FWebSocketMessageTransportSerializeHelper
+{
+public:
+	static const TMap<EMessageScope, FString> MessageScopeStringMapping;
+
+	static bool Serialize(const TSharedRef<IMessageContext, ESPMode::ThreadSafe>& Context, FString& OutJsonMessage)
+	{
+		TSharedRef<FJsonObject> JsonRoot = MakeShared<FJsonObject>();
+		JsonRoot->SetStringField(WebSocketMessaging::Tag::Sender, Context->GetSender().ToString());
+		TArray<TSharedPtr<FJsonValue>> JsonRecipients;
+		for (const FMessageAddress& Recipient : Context->GetRecipients())
+		{
+			JsonRecipients.Add(MakeShared<FJsonValueString>(Recipient.ToString()));
+		}
+		JsonRoot->SetArrayField(WebSocketMessaging::Tag::Recipients, JsonRecipients);
+		JsonRoot->SetStringField(WebSocketMessaging::Tag::MessageType, Context->GetMessageTypePathName().ToString());
+		JsonRoot->SetNumberField(WebSocketMessaging::Tag::Expiration, Context->GetExpiration().ToUnixTimestamp());
+		JsonRoot->SetNumberField(WebSocketMessaging::Tag::TimeSent, Context->GetTimeSent().ToUnixTimestamp());
+		JsonRoot->SetStringField(WebSocketMessaging::Tag::Scope, MessageScopeStringMapping[Context->GetScope()]);
+
+
+		TSharedRef<FJsonObject> JsonAnnotations = MakeShared<FJsonObject>();
+		for (const TPair<FName, FString>& Pair : Context->GetAnnotations())
+		{
+			JsonAnnotations->SetStringField(Pair.Key.ToString(), Pair.Value);
+		}
+		JsonRoot->SetObjectField(WebSocketMessaging::Tag::Annotations, JsonAnnotations);
+
+		TSharedRef<FJsonObject> OutJsonObject = MakeShared<FJsonObject>();
+
+		// Remark: This will change the case of field names, see StandardizeCase.
+		if (!FJsonObjectConverter::UStructToJsonObject(Context->GetMessageTypeInfo().Get(), Context->GetMessage(), OutJsonObject))
+		{
+			return false;
+		}
+
+		JsonRoot->SetObjectField(WebSocketMessaging::Tag::Message, OutJsonObject);
+
+		TSharedRef<TJsonWriter<>> Writer = TJsonWriterFactory<>::Create(&OutJsonMessage);
+		return FJsonSerializer::Serialize(JsonRoot, Writer);
+	}
+
+	static bool Serialize(const TSharedRef<IMessageContext, ESPMode::ThreadSafe>& Context, FArrayWriter& OutCborBinaryWriter)
+	{
+		FCborHeader Header(ECborCode::Map | ECborCode::Indefinite);
+		OutCborBinaryWriter << Header;
+			
+		{
+			FCborWriter CborWriter(&OutCborBinaryWriter);
+			
+			CborWriter.WriteValue(FString(WebSocketMessaging::Tag::Sender));
+			CborWriter.WriteValue(Context->GetSender().ToString());
+			
+			CborWriter.WriteValue(FString(WebSocketMessaging::Tag::Recipients));
+			CborWriter.WriteContainerStart(ECborCode::Array, -1);
+			for (const FMessageAddress& Recipient : Context->GetRecipients())
+			{
+				CborWriter.WriteValue(Recipient.ToString());
+			}
+			CborWriter.WriteContainerEnd();
+			
+			CborWriter.WriteValue(FString(WebSocketMessaging::Tag::MessageType));
+			CborWriter.WriteValue(Context->GetMessageTypePathName().ToString());
+			
+			CborWriter.WriteValue(FString(WebSocketMessaging::Tag::Expiration));
+			CborWriter.WriteValue(Context->GetExpiration().ToUnixTimestamp());
+			
+			CborWriter.WriteValue(FString(WebSocketMessaging::Tag::TimeSent));
+			CborWriter.WriteValue(Context->GetTimeSent().ToUnixTimestamp());
+			
+			CborWriter.WriteValue(FString(WebSocketMessaging::Tag::Scope));
+			CborWriter.WriteValue(MessageScopeStringMapping[Context->GetScope()]);
+			
+			
+			CborWriter.WriteValue(FString(WebSocketMessaging::Tag::Annotations));
+			CborWriter.WriteContainerStart(ECborCode::Map, -1);
+			for (const TPair<FName, FString>& Annotation : Context->GetAnnotations())
+			{
+				CborWriter.WriteValue(Annotation.Key.ToString());
+				CborWriter.WriteValue(Annotation.Value);
+			}
+			CborWriter.WriteContainerEnd();
+			
+			CborWriter.WriteValue(FString(WebSocketMessaging::Tag::Message));
+		}
+		FCborStructSerializerBackend Backend(OutCborBinaryWriter, EStructSerializerBackendFlags::Default);
+		FStructSerializer::Serialize(Context->GetMessage(), *Context->GetMessageTypeInfo().Get(), Backend);
+			
+		Header.Set(ECborCode::Break);
+		OutCborBinaryWriter << Header;
+
+		return true;
+	}
+
+	template<typename OutputType>
+	struct TOnDemandSerializer
+	{
+		OutputType OutputMessage;
+		bool bIsSerializeAttempted = false;
+		bool bIsSerialized = false;
+
+		bool SerializeOnDemand(const TSharedRef<IMessageContext, ESPMode::ThreadSafe>& Context)
+		{
+			if(!bIsSerializeAttempted)
+			{
+				bIsSerializeAttempted = true;
+				bIsSerialized = FWebSocketMessageTransportSerializeHelper::Serialize(Context, OutputMessage);
+			}
+			return bIsSerialized;
+		}
+	};
+};
+
+const TMap<EMessageScope, FString> FWebSocketMessageTransportSerializeHelper::MessageScopeStringMapping =
+{
+	{EMessageScope::Thread, "Thread"},
+	{EMessageScope::Process, "Process"},
+	{EMessageScope::Network, "Network"},
+	{EMessageScope::All, "All"}
+};
+
+bool FWebSocketMessageTransport::TransportMessage(const TSharedRef<IMessageContext, ESPMode::ThreadSafe>& Context, const TArray<FGuid>& Recipients)
+{
+	TMap<FGuid, FWebSocketMessageConnectionRef> RecipientConnections;
+
+	if (Recipients.Num() == 0)
+	{
+		// broadcast the message to all valid connections
+		RecipientConnections = WebSocketMessageConnections.FilterByPredicate([](const TPair<FGuid, FWebSocketMessageConnectionRef>& Pair) -> bool
+			{
+				return !Pair.Value->bDestroyed && Pair.Value->IsConnected();
+			});
+	}
+	else
+	{
+		// Find connections for each recipient.  We do not transport unicast messages for unknown nodes.
+		for (const FGuid& Recipient : Recipients)
+		{
+			FWebSocketMessageConnectionRef* RecipientConnection = WebSocketMessageConnections.Find(Recipient);
+			if (RecipientConnection && !(*RecipientConnection)->bDestroyed && (*RecipientConnection)->IsConnected())
+			{
+				RecipientConnections.Add(Recipient, *RecipientConnection);
+			}
+		}
+	}
+
+	if (RecipientConnections.Num() == 0)
+	{
+		return false;
+	}
+
+	// Remark: Json serializer uses UStructToJsonObject, which is going to change
+	// the field name case (see StandardizeCase), which is going to cause a difference
+	// with the field names in Cbor (which doesn't change the case).
+
+	FWebSocketMessageTransportSerializeHelper::TOnDemandSerializer<FString> JsonSerializer;
+	FWebSocketMessageTransportSerializeHelper::TOnDemandSerializer<FArrayWriter> CborSerializer;
+	
+	// Serialize the message on demand in the appropriate format for each peer connections.
+	for (const TPair<FGuid, FWebSocketMessageConnectionRef>& Pair : RecipientConnections)
+	{
+		if (Pair.Value->WebSocketConnection.IsValid())
+		{
+			// Remark: client connections are always text/json
+			if(JsonSerializer.SerializeOnDemand(Context))
+			{
+				Pair.Value->WebSocketConnection->Send(JsonSerializer.OutputMessage);
+			}
+		}
+		else if (Pair.Value->WebSocketServerConnection)
+		{
+			// Remark: server connections are always binary.
+			const UWebSocketMessagingSettings* Settings = GetDefault<UWebSocketMessagingSettings>();
+
+			if(Settings->ServerTransportFormat == EWebSocketMessagingTransportFormat::Json)
+			{
+				if(JsonSerializer.SerializeOnDemand(Context))
+				{
+					FTCHARToUTF8 Converted(*JsonSerializer.OutputMessage);
+					Pair.Value->WebSocketServerConnection->Send(reinterpret_cast<const uint8*>(Converted.Get()), Converted.Length(), false);
+				}
+			}
+			else
+			{
+				if(CborSerializer.SerializeOnDemand(Context))
+				{
+					Pair.Value->WebSocketServerConnection->Send(CborSerializer.OutputMessage.GetData(), CborSerializer.OutputMessage.Num(), false);
+				}
+			}
+		}
+	}
+
+	return true;
+}
+
+void FWebSocketMessageTransport::OnConnected(FWebSocketMessageConnectionRef WebSocketMessageConnection)
+{
+	UE_LOG(LogWebSocketMessaging, Log, TEXT("Connected to %s"), *WebSocketMessageConnection->Url);
+	WebSocketMessageConnection->bIsConnecting = false;
+}
+
+void FWebSocketMessageTransport::OnServerConnectionClosed(FWebSocketMessageConnectionRef WebSocketMessageConnection)
+{
+	UE_LOG(LogWebSocketMessaging, Log, TEXT("%s disconnected"), *WebSocketMessageConnection->Url);
+	ForgetTransportNode(WebSocketMessageConnection);
+	WebSocketMessageConnections.Remove(WebSocketMessageConnection->Guid);
+}
+
+void FWebSocketMessageTransport::RetryConnection(FWebSocketMessageConnectionRef WebSocketMessageConnection)
+{
+	if (WebSocketMessageConnection->bIsConnecting)
+	{
+		return;
+	}
+
+	WebSocketMessageConnection->RetryHandle = FTSTicker::GetCoreTicker().AddTicker(FTickerDelegate::CreateLambda([](float DeltaTime, FWebSocketMessageConnectionRef WebSocketMessageConnection)
+		{
+			if (!WebSocketMessageConnection->bDestroyed && !WebSocketMessageConnection->bIsConnecting && !WebSocketMessageConnection->WebSocketConnection->IsConnected())
+			{
+				WebSocketMessageConnection->bIsConnecting = true;
+				WebSocketMessageConnection->WebSocketConnection->Connect();
+			}
+			return false;
+		}, WebSocketMessageConnection), 1.0f);
+}
+
+void FWebSocketMessageTransport::ClientConnected(INetworkingWebSocket* NetworkingWebSocket)
+{
+	FString RemoteEndPoint = NetworkingWebSocket->RemoteEndPoint(true);
+	UE_LOG(LogWebSocketMessaging, Log, TEXT("New WebSocket Server connection: %s"), *RemoteEndPoint);
+
+	FGuid Guid = FGuid::NewGuid();
+
+	FWebSocketMessageConnectionRef WebSocketMessageConnection = MakeShared<FWebSocketMessageConnection>(RemoteEndPoint, Guid, NetworkingWebSocket);
+
+	NetworkingWebSocket->SetReceiveCallBack(FWebSocketPacketReceivedCallBack::CreateThreadSafeSP(this, &FWebSocketMessageTransport::OnServerJsonMessage, WebSocketMessageConnection));
+	NetworkingWebSocket->SetSocketClosedCallBack(FWebSocketInfoCallBack::CreateThreadSafeSP(this, &FWebSocketMessageTransport::OnServerConnectionClosed, WebSocketMessageConnection));
+	NetworkingWebSocket->SetErrorCallBack(FWebSocketInfoCallBack::CreateThreadSafeSP(this, &FWebSocketMessageTransport::OnServerConnectionClosed, WebSocketMessageConnection));
+
+	WebSocketMessageConnections.Add(Guid, WebSocketMessageConnection);
+}
+
+bool FWebSocketMessageTransport::ServerTick(float DeltaTime)
+{
+	if (Server.IsValid())
+	{
+		Server->Tick();
+	}
+
+	return true;
+}
+
+void FWebSocketMessageTransport::ForgetTransportNode(FWebSocketMessageConnectionRef WebSocketMessageConnection)
+{
+	if(TransportHandler)
+	{
+		TransportHandler->ForgetTransportNode(WebSocketMessageConnection->Guid);
+	}
+}
diff --git a/Engine/Plugins/Experimental/WebSocketMessaging/Source/WebSocketMessaging/Private/WebSocketMessaging.cpp b/Engine/Plugins/Experimental/WebSocketMessaging/Source/WebSocketMessaging/Private/WebSocketMessaging.cpp
new file mode 100644
--- /dev/null
+++ b/Engine/Plugins/Experimental/WebSocketMessaging/Source/WebSocketMessaging/Private/WebSocketMessaging.cpp
@@ -0,0 +1,93 @@
+// Copyright Epic Games, Inc. All Rights Reserved.
+
+#include "WebSocketMessaging.h"
+#include "IMessageBridge.h"
+#if WITH_EDITOR
+#include "ISettingsModule.h"
+#include "ISettingsSection.h"
+#endif
+#include "MessageBridgeBuilder.h"
+#include "WebSocketMessageTransport.h"
+#include "WebSocketMessagingSettings.h"
+
+DEFINE_LOG_CATEGORY(LogWebSocketMessaging);
+
+
+#define LOCTEXT_NAMESPACE "FWebSocketMessagingModule"
+
+void FWebSocketMessagingModule::StartupModule()
+{
+	if (!FModuleManager::Get().LoadModule(TEXT("WebSockets")))
+	{
+		return;
+	}
+#if WITH_EDITOR
+	// register settings
+	ISettingsModule* SettingsModule = FModuleManager::GetModulePtr<ISettingsModule>("Settings");
+
+	if (SettingsModule != nullptr)
+	{
+		ISettingsSectionPtr SettingsSection = SettingsModule->RegisterSettings("Project", "Plugins", "WebSocketMessaging",
+			LOCTEXT("WebSocketMessagingSettingsName", "WebSocket Messaging"),
+			LOCTEXT("WebSocketMessagingSettingsDescription", "Configure the WebSocket Messaging plug-in."),
+			GetMutableDefault<UWebSocketMessagingSettings>()
+		);
+
+		if (SettingsSection.IsValid())
+		{
+			SettingsSection->OnModified().BindRaw(this, &FWebSocketMessagingModule::HandleSettingsSaved);
+		}
+	}
+#endif // WITH_EDITOR
+
+	// trigger initialization
+	HandleSettingsSaved();
+}
+
+void FWebSocketMessagingModule::ShutdownModule()
+{
+#if WITH_EDITOR
+	// unregister settings
+	ISettingsModule* SettingsModule = FModuleManager::GetModulePtr<ISettingsModule>("Settings");
+
+	if (SettingsModule != nullptr)
+	{
+		SettingsModule->UnregisterSettings("Project", "Plugins", "WebSocketMessaging");
+	}
+#endif
+}
+
+bool FWebSocketMessagingModule::HandleSettingsSaved()
+{
+	if (GetDefault<UWebSocketMessagingSettings>()->EnableTransport)
+	{
+		InitializeBridge();
+	}
+	else
+	{
+		ShutdownBridge();
+	}
+
+	return true;
+}
+
+void FWebSocketMessagingModule::InitializeBridge()
+{
+	ShutdownBridge();
+
+	MessageBridge = FMessageBridgeBuilder().UsingTransport(MakeShared<FWebSocketMessageTransport>()).Build();
+}
+
+void FWebSocketMessagingModule::ShutdownBridge()
+{
+	if (MessageBridge.IsValid())
+	{
+		MessageBridge->Disable();
+		FPlatformProcess::Sleep(0.1f);
+		MessageBridge.Reset();
+	}
+}
+
+#undef LOCTEXT_NAMESPACE
+
+IMPLEMENT_MODULE(FWebSocketMessagingModule, WebSocketMessaging)
\ No newline at end of file
diff --git a/Engine/Plugins/Experimental/WebSocketMessaging/Source/WebSocketMessaging/Public/WebSocketDeserializedMessage.h b/Engine/Plugins/Experimental/WebSocketMessaging/Source/WebSocketMessaging/Public/WebSocketDeserializedMessage.h
new file mode 100644
--- /dev/null
+++ b/Engine/Plugins/Experimental/WebSocketMessaging/Source/WebSocketMessaging/Public/WebSocketDeserializedMessage.h
@@ -0,0 +1,101 @@
+// Copyright Epic Games, Inc. All Rights Reserved.
+
+#pragma once
+
+#include "CoreMinimal.h"
+#include "IMessageAttachment.h"
+#include "IMessageContext.h"
+
+
+/**
+ *
+ */
+class WEBSOCKETMESSAGING_API FWebSocketDeserializedMessage : public IMessageContext
+{
+public:
+	FWebSocketDeserializedMessage();
+	~FWebSocketDeserializedMessage();
+
+	virtual bool ParseJson(const FString& Json);
+
+	virtual TSharedPtr<IMessageAttachment, ESPMode::ThreadSafe> GetAttachment() const override
+	{
+		return nullptr;
+	}
+
+	virtual const FDateTime& GetExpiration() const override
+	{
+		return Expiration;
+	}
+
+	virtual const void* GetMessage() const override
+	{
+		return Message;
+	}
+
+	virtual const TWeakObjectPtr<UScriptStruct, FWeakObjectPtr>& GetMessageTypeInfo() const override
+	{
+		return TypeInfo;
+	}
+
+	virtual const FMessageAddress& GetSender() const override
+	{
+		return Sender;
+	}
+
+	virtual const FDateTime& GetTimeSent() const override
+	{
+		return TimeSent;
+	}
+
+	virtual const TMap<FName, FString>& GetAnnotations() const override
+	{
+		return Annotations;
+	}
+
+	virtual TSharedPtr<IMessageContext, ESPMode::ThreadSafe> GetOriginalContext() const override
+	{
+		return nullptr;
+	}
+
+	virtual EMessageFlags GetFlags() const override
+	{
+		return EMessageFlags::None;
+	}
+
+	virtual const FMessageAddress& GetForwarder() const override
+	{
+		return Sender;
+	}
+
+	virtual ENamedThreads::Type GetSenderThread() const override
+	{
+		return ENamedThreads::AnyThread;
+	}
+
+	virtual const TArray<FMessageAddress, FDefaultAllocator>& GetRecipients() const override
+	{
+		return Recipients;
+	}
+
+	virtual const FDateTime& GetTimeForwarded() const override
+	{
+		return TimeSent;
+	}
+
+	virtual EMessageScope GetScope() const override
+	{
+		return Scope;
+	}
+
+protected:
+
+	FDateTime Expiration;
+	FDateTime TimeSent;
+	void* Message;
+	TWeakObjectPtr<UScriptStruct> TypeInfo;
+	FMessageAddress Sender;
+	TMap<FName, FString> Annotations;
+	TArray<FMessageAddress> Recipients;
+	EMessageScope Scope;
+};
diff --git a/Engine/Plugins/Experimental/WebSocketMessaging/Source/WebSocketMessaging/Public/WebSocketMessageTransport.h b/Engine/Plugins/Experimental/WebSocketMessaging/Source/WebSocketMessaging/Public/WebSocketMessageTransport.h
new file mode 100644
--- /dev/null
+++ b/Engine/Plugins/Experimental/WebSocketMessaging/Source/WebSocketMessaging/Public/WebSocketMessageTransport.h
@@ -0,0 +1,104 @@
+// Copyright Epic Games, Inc. All Rights Reserved.
+
+#pragma once
+
+#include "CoreMinimal.h"
+#include "Containers/Ticker.h"
+#include "IMessageTransport.h"
+
+struct FWebSocketMessageConnection
+{
+	FWebSocketMessageConnection() = delete;
+	FWebSocketMessageConnection(const FString& InUrl, FGuid InGuid, TSharedRef<class IWebSocket, ESPMode::ThreadSafe> InWebSocketConnection) :
+		Url(InUrl),
+		Guid(InGuid),
+		WebSocketConnection(InWebSocketConnection),
+		WebSocketServerConnection(nullptr),
+		bIsConnecting(true),
+		bDestroyed(false)
+	{
+
+	}
+
+	FWebSocketMessageConnection(const FString& InUrl, FGuid InGuid, class INetworkingWebSocket* InWebSocketServerConnection) :
+		Url(InUrl),
+		Guid(InGuid),
+		WebSocketConnection(nullptr),
+		WebSocketServerConnection(InWebSocketServerConnection),
+		bIsConnecting(true),
+		bDestroyed(false)
+	{
+
+	}
+
+	bool IsConnected() const;
+
+	void Close();
+
+	/** The WebSocket url */
+	const FString Url;
+
+	/** The message transport Guid */
+	const FGuid Guid;
+
+	/** Reference to the client websocket connection */
+	TSharedPtr<class IWebSocket, ESPMode::ThreadSafe> WebSocketConnection;
+
+	/** Reference to the server websocket connection */
+	class INetworkingWebSocket* WebSocketServerConnection;
+
+	/** Is the socket still trying to connect? */
+	bool bIsConnecting;
+
+	/** The socket is about to be destroyed */
+	bool bDestroyed;
+
+	/** Retry timer */
+	FTSTicker::FDelegateHandle RetryHandle;
+};
+
+using FWebSocketMessageConnectionRef = TSharedRef<FWebSocketMessageConnection, ESPMode::ThreadSafe>;
+
+/**
+ * 
+ */
+class WEBSOCKETMESSAGING_API FWebSocketMessageTransport : public IMessageTransport, public TSharedFromThis<FWebSocketMessageTransport>
+{
+public:
+	FWebSocketMessageTransport();
+	~FWebSocketMessageTransport();
+
+	virtual FName GetDebugName() const override
+	{
+		return "WebSocketMessageTransport";
+	}
+
+	virtual bool StartTransport(IMessageTransportHandler& Handler) override;
+
+	virtual void StopTransport() override;
+
+	virtual bool TransportMessage(const TSharedRef<IMessageContext, ESPMode::ThreadSafe>& Context, const TArray<FGuid>& Recipients) override;
+
+	virtual void OnJsonMessage(const FString& Message, FWebSocketMessageConnectionRef WebSocketMessageConnection);
+	virtual void OnConnected(FWebSocketMessageConnectionRef WebSocketMessageConnection);
+	virtual void OnClosed(int32 Code, const FString& Reason, bool bUserClose, FWebSocketMessageConnectionRef WebSocketMessageConnection);
+	virtual void OnConnectionError(const FString& Message, FWebSocketMessageConnectionRef WebSocketMessageConnection);
+
+	virtual void RetryConnection(FWebSocketMessageConnectionRef WebSocketMessageConnection);
+
+	virtual void ClientConnected(class INetworkingWebSocket* NetworkingWebSocket);
+
+	virtual bool ServerTick(float DeltaTime);
+
+	virtual void OnServerJsonMessage(void* Data, int32 DataSize, FWebSocketMessageConnectionRef WebSocketMessageConnection);
+	virtual void OnServerConnectionClosed(FWebSocketMessageConnectionRef WebSocketMessageConnection);
+
+protected:
+	void ForgetTransportNode(FWebSocketMessageConnectionRef WebSocketMessageConnection);
+	
+	IMessageTransportHandler* TransportHandler;
+	TMap<FGuid, FWebSocketMessageConnectionRef> WebSocketMessageConnections;
+
+	TUniquePtr<class IWebSocketServer> Server;
+	FTSTicker::FDelegateHandle ServerTickerHandle;
+};
diff --git a/Engine/Plugins/Experimental/WebSocketMessaging/Source/WebSocketMessaging/Public/WebSocketMessaging.h b/Engine/Plugins/Experimental/WebSocketMessaging/Source/WebSocketMessaging/Public/WebSocketMessaging.h
new file mode 100644
--- /dev/null
+++ b/Engine/Plugins/Experimental/WebSocketMessaging/Source/WebSocketMessaging/Public/WebSocketMessaging.h
@@ -0,0 +1,46 @@
+// Copyright Epic Games, Inc. All Rights Reserved.
+
+#pragma once
+
+#include "CoreMinimal.h"
+#include "Modules/ModuleManager.h"
+
+DECLARE_LOG_CATEGORY_EXTERN(LogWebSocketMessaging, Log, All);
+
+namespace WebSocketMessaging
+{
+	namespace Tag
+	{
+		static constexpr const TCHAR* Sender = TEXT("Sender");
+		static constexpr const TCHAR* Recipients = TEXT("Recipients");
+		static constexpr const TCHAR* Expiration = TEXT("Expiration");
+		static constexpr const TCHAR* TimeSent = TEXT("TimeSent");
+		static constexpr const TCHAR* Annotations = TEXT("Annotations");
+		static constexpr const TCHAR* Scope = TEXT("Scope");
+		static constexpr const TCHAR* MessageType = TEXT("MessageType");
+		static constexpr const TCHAR* Message = TEXT("Message");
+	}
+
+	namespace Header
+	{
+		static constexpr const TCHAR* TransportId = TEXT("UE-MessageBus-TransportId");
+	}
+}
+
+class FWebSocketMessagingModule : public IModuleInterface
+{
+public:
+
+	/** IModuleInterface implementation */
+	virtual void StartupModule() override;
+	virtual void ShutdownModule() override;
+
+	virtual bool HandleSettingsSaved();
+
+	virtual void InitializeBridge();
+	virtual void ShutdownBridge();
+
+protected:
+	/** Holds the message bridge if present. */
+	TSharedPtr<class IMessageBridge, ESPMode::ThreadSafe> MessageBridge;
+};
diff --git a/Engine/Plugins/Experimental/WebSocketMessaging/Source/WebSocketMessaging/Public/WebSocketMessagingSettings.h b/Engine/Plugins/Experimental/WebSocketMessaging/Source/WebSocketMessaging/Public/WebSocketMessagingSettings.h
new file mode 100644
--- /dev/null
+++ b/Engine/Plugins/Experimental/WebSocketMessaging/Source/WebSocketMessaging/Public/WebSocketMessagingSettings.h
@@ -0,0 +1,45 @@
+// Copyright Epic Games, Inc. All Rights Reserved.
+
+#pragma once
+
+#include "CoreMinimal.h"
+#include "UObject/NoExportTypes.h"
+#include "WebSocketMessagingSettings.generated.h"
+
+UENUM()
+enum class EWebSocketMessagingTransportFormat : uint8
+{
+	Json,
+	Cbor
+};
+
+/**
+ * 
+ */
+UCLASS(config=Engine)
+class WEBSOCKETMESSAGING_API UWebSocketMessagingSettings : public UObject
+{
+	GENERATED_BODY()
+
+public:
+
+	/** Whether the WebSocket transport channel is enabled */
+	UPROPERTY(config, EditAnywhere, Category = Transport)
+	bool EnableTransport = false;
+
+	/** Bind the WebSocket server on the specified port (0 disables it) */
+	UPROPERTY(config, EditAnywhere, Category = Transport)
+	int32 ServerPort = 0;
+
+	/** Format used to serialize the messages on the server's WebSockets.*/
+	UPROPERTY(Config, EditAnywhere, Category = Transport)
+	EWebSocketMessagingTransportFormat ServerTransportFormat = EWebSocketMessagingTransportFormat::Cbor;
+	
+    /** The WebSocket Urls to connect to (Eg. ws://example.com/xyz) */
+    UPROPERTY(config, EditAnywhere, Category = Transport)
+    TArray<FString> ConnectToEndpoints;
+
+	/** Additional HTTP headers to set when connecting to endpoints */
+	UPROPERTY(config, EditAnywhere, Category = Transport)
+	TMap<FString, FString> HttpHeaders;
+};
diff --git a/Engine/Plugins/Experimental/WebSocketMessaging/Source/WebSocketMessaging/WebSocketMessaging.Build.cs b/Engine/Plugins/Experimental/WebSocketMessaging/Source/WebSocketMessaging/WebSocketMessaging.Build.cs
new file mode 100644
--- /dev/null
+++ b/Engine/Plugins/Experimental/WebSocketMessaging/Source/WebSocketMessaging/WebSocketMessaging.Build.cs
@@ -0,0 +1,35 @@
+// Copyright Epic Games, Inc. All Rights Reserved.
+
+using UnrealBuildTool;
+
+public class WebSocketMessaging : ModuleRules
+{
+	public WebSocketMessaging(ReadOnlyTargetRules Target) : base(Target)
+	{
+		PCHUsage = ModuleRules.PCHUsageMode.UseExplicitOrSharedPCHs;
+
+
+		PublicDependencyModuleNames.AddRange(
+			new string[]
+			{
+				"Core",
+			}
+			);
+
+
+		PrivateDependencyModuleNames.AddRange(
+			new string[]
+			{
+				"CoreUObject",
+				"Engine",
+				"Messaging",
+				"WebSockets",
+				"WebSocketNetworking",
+				"Json",
+				"JsonUtilities",
+				"Cbor",
+				"Serialization"
+			}
+			);
+	}
+}
diff --git a/Engine/Plugins/Experimental/WebSocketMessaging/WebSocketMessaging.uplugin b/Engine/Plugins/Experimental/WebSocketMessaging/WebSocketMessaging.uplugin
new file mode 100644
--- /dev/null
+++ b/Engine/Plugins/Experimental/WebSocketMessaging/WebSocketMessaging.uplugin
@@ -0,0 +1,35 @@
+{
+	"FileVersion": 3,
+	"Version": 1,
+	"VersionName": "1.0",
+	"FriendlyName": "WebSocketMessaging",
+	"Description": "",
+	"Category": "Other",
+	"CreatedBy": "",
+	"CreatedByURL": "",
+	"DocsURL": "",
+	"MarketplaceURL": "",
+	"SupportURL": "",
+	"CanContainContent": false,
+	"IsBetaVersion": false,
+	"IsExperimentalVersion": false,
+	"Installed": false,
+	"Modules": [
+		{
+			"Name": "WebSocketMessaging",
+			"Type": "Runtime",
+			"LoadingPhase": "Default",
+			"PlatformAllowList": [
+				"Mac",
+				"Win64",
+				"Linux"
+			]
+		}
+	],
+	"Plugins": [
+		{
+			"Name": "WebSocketNetworking",
+			"Enabled": true
+		}
+	]
+}
\ No newline at end of file
