Web Socket Messaging:
- public api cleanup

#changelist validated
#virtualized

diff --git a/Engine/Plugins/Experimental/WebSocketMessaging/Source/WebSocketMessaging/Private/WebSocketDeserializedMessage.h b/Engine/Plugins/Experimental/WebSocketMessaging/Source/WebSocketMessaging/Private/WebSocketDeserializedMessage.h
new file mode 100644
--- /dev/null
+++ b/Engine/Plugins/Experimental/WebSocketMessaging/Source/WebSocketMessaging/Private/WebSocketDeserializedMessage.h
@@ -0,0 +1,102 @@
+// Copyright Epic Games, Inc. All Rights Reserved.
+
+#pragma once
+
+#include "Containers/Array.h"
+#include "Containers/Map.h"
+#include "Containers/UnrealString.h"
+#include "IMessageAttachment.h"
+#include "IMessageContext.h"
+#include "UObject/NameTypes.h"
+#include "UObject/WeakObjectPtr.h"
+#include "UObject/WeakObjectPtrTemplates.h"
+
+class FWebSocketDeserializedMessage : public IMessageContext
+{
+public:
+	FWebSocketDeserializedMessage();
+
+	virtual ~FWebSocketDeserializedMessage() override;
+
+	virtual bool ParseJson(const FString& Json);
+
+	virtual TSharedPtr<IMessageAttachment, ESPMode::ThreadSafe> GetAttachment() const override
+	{
+		return nullptr;
+	}
+
+	virtual const FDateTime& GetExpiration() const override
+	{
+		return Expiration;
+	}
+
+	virtual const void* GetMessage() const override
+	{
+		return Message;
+	}
+
+	virtual const TWeakObjectPtr<UScriptStruct, FWeakObjectPtr>& GetMessageTypeInfo() const override
+	{
+		return TypeInfo;
+	}
+
+	virtual const FMessageAddress& GetSender() const override
+	{
+		return Sender;
+	}
+
+	virtual const FDateTime& GetTimeSent() const override
+	{
+		return TimeSent;
+	}
+
+	virtual const TMap<FName, FString>& GetAnnotations() const override
+	{
+		return Annotations;
+	}
+
+	virtual TSharedPtr<IMessageContext, ESPMode::ThreadSafe> GetOriginalContext() const override
+	{
+		return nullptr;
+	}
+
+	virtual EMessageFlags GetFlags() const override
+	{
+		return EMessageFlags::None;
+	}
+
+	virtual const FMessageAddress& GetForwarder() const override
+	{
+		return Sender;
+	}
+
+	virtual ENamedThreads::Type GetSenderThread() const override
+	{
+		return ENamedThreads::AnyThread;
+	}
+
+	virtual const TArray<FMessageAddress, FDefaultAllocator>& GetRecipients() const override
+	{
+		return Recipients;
+	}
+
+	virtual const FDateTime& GetTimeForwarded() const override
+	{
+		return TimeSent;
+	}
+
+	virtual EMessageScope GetScope() const override
+	{
+		return Scope;
+	}
+
+protected:
+	FDateTime Expiration;
+	FDateTime TimeSent;
+	void* Message;
+	TWeakObjectPtr<UScriptStruct> TypeInfo;
+	FMessageAddress Sender;
+	TMap<FName, FString> Annotations;
+	TArray<FMessageAddress> Recipients;
+	EMessageScope Scope;
+};
diff --git a/Engine/Plugins/Experimental/WebSocketMessaging/Source/WebSocketMessaging/Private/WebSocketMessageTransport.h b/Engine/Plugins/Experimental/WebSocketMessaging/Source/WebSocketMessaging/Private/WebSocketMessageTransport.h
new file mode 100644
--- /dev/null
+++ b/Engine/Plugins/Experimental/WebSocketMessaging/Source/WebSocketMessaging/Private/WebSocketMessageTransport.h
@@ -0,0 +1,103 @@
+// Copyright Epic Games, Inc. All Rights Reserved.
+
+#pragma once
+
+#include "Containers/Map.h"
+#include "Containers/Ticker.h"
+#include "IMessageTransport.h"
+#include "Templates/SharedPointer.h"
+#include "Templates/UniquePtr.h"
+
+struct FWebSocketMessageConnection
+{
+	FWebSocketMessageConnection() = delete;
+	FWebSocketMessageConnection(const FString& InUrl, FGuid InGuid, TSharedRef<class IWebSocket, ESPMode::ThreadSafe> InWebSocketConnection) :
+		Url(InUrl),
+		Guid(InGuid),
+		WebSocketConnection(InWebSocketConnection),
+		WebSocketServerConnection(nullptr),
+		bIsConnecting(true),
+		bDestroyed(false)
+	{
+
+	}
+
+	FWebSocketMessageConnection(const FString& InUrl, FGuid InGuid, class INetworkingWebSocket* InWebSocketServerConnection) :
+		Url(InUrl),
+		Guid(InGuid),
+		WebSocketConnection(nullptr),
+		WebSocketServerConnection(InWebSocketServerConnection),
+		bIsConnecting(true),
+		bDestroyed(false)
+	{
+
+	}
+
+	bool IsConnected() const;
+
+	void Close();
+
+	/** The WebSocket url */
+	const FString Url;
+
+	/** The message transport Guid */
+	const FGuid Guid;
+
+	/** Reference to the client websocket connection */
+	TSharedPtr<class IWebSocket, ESPMode::ThreadSafe> WebSocketConnection;
+
+	/** Reference to the server websocket connection */
+	class INetworkingWebSocket* WebSocketServerConnection;
+
+	/** Is the socket still trying to connect? */
+	bool bIsConnecting;
+
+	/** The socket is about to be destroyed */
+	bool bDestroyed;
+
+	/** Retry timer */
+	FTSTicker::FDelegateHandle RetryHandle;
+};
+
+using FWebSocketMessageConnectionRef = TSharedRef<FWebSocketMessageConnection, ESPMode::ThreadSafe>;
+
+class FWebSocketMessageTransport : public IMessageTransport, public TSharedFromThis<FWebSocketMessageTransport>
+{
+public:
+	FWebSocketMessageTransport();
+	virtual ~FWebSocketMessageTransport() override;
+
+	virtual FName GetDebugName() const override
+	{
+		return "WebSocketMessageTransport";
+	}
+
+	virtual bool StartTransport(IMessageTransportHandler& Handler) override;
+
+	virtual void StopTransport() override;
+
+	virtual bool TransportMessage(const TSharedRef<IMessageContext, ESPMode::ThreadSafe>& Context, const TArray<FGuid>& Recipients) override;
+
+	virtual void OnJsonMessage(const FString& Message, FWebSocketMessageConnectionRef WebSocketMessageConnection);
+	virtual void OnConnected(FWebSocketMessageConnectionRef WebSocketMessageConnection);
+	virtual void OnClosed(int32 Code, const FString& Reason, bool bUserClose, FWebSocketMessageConnectionRef WebSocketMessageConnection);
+	virtual void OnConnectionError(const FString& Message, FWebSocketMessageConnectionRef WebSocketMessageConnection);
+
+	virtual void RetryConnection(FWebSocketMessageConnectionRef WebSocketMessageConnection);
+
+	virtual void ClientConnected(class INetworkingWebSocket* NetworkingWebSocket);
+
+	virtual bool ServerTick(float DeltaTime);
+
+	virtual void OnServerJsonMessage(void* Data, int32 DataSize, FWebSocketMessageConnectionRef WebSocketMessageConnection);
+	virtual void OnServerConnectionClosed(FWebSocketMessageConnectionRef WebSocketMessageConnection);
+
+protected:
+	void ForgetTransportNode(FWebSocketMessageConnectionRef WebSocketMessageConnection);
+	
+	IMessageTransportHandler* TransportHandler;
+	TMap<FGuid, FWebSocketMessageConnectionRef> WebSocketMessageConnections;
+
+	TUniquePtr<class IWebSocketServer> Server;
+	FTSTicker::FDelegateHandle ServerTickerHandle;
+};
diff --git a/Engine/Plugins/Experimental/WebSocketMessaging/Source/WebSocketMessaging/Private/WebSocketMessaging.h b/Engine/Plugins/Experimental/WebSocketMessaging/Source/WebSocketMessaging/Private/WebSocketMessaging.h
new file mode 100644
--- /dev/null
+++ b/Engine/Plugins/Experimental/WebSocketMessaging/Source/WebSocketMessaging/Private/WebSocketMessaging.h
@@ -0,0 +1,47 @@
+// Copyright Epic Games, Inc. All Rights Reserved.
+
+#pragma once
+
+#include "Logging/LogMacros.h"
+#include "Modules/ModuleManager.h"
+#include "Templates/SharedPointer.h"
+
+DECLARE_LOG_CATEGORY_EXTERN(LogWebSocketMessaging, Log, All);
+
+namespace WebSocketMessaging
+{
+	namespace Tag
+	{
+		static constexpr const TCHAR* Sender = TEXT("Sender");
+		static constexpr const TCHAR* Recipients = TEXT("Recipients");
+		static constexpr const TCHAR* Expiration = TEXT("Expiration");
+		static constexpr const TCHAR* TimeSent = TEXT("TimeSent");
+		static constexpr const TCHAR* Annotations = TEXT("Annotations");
+		static constexpr const TCHAR* Scope = TEXT("Scope");
+		static constexpr const TCHAR* MessageType = TEXT("MessageType");
+		static constexpr const TCHAR* Message = TEXT("Message");
+	}
+
+	namespace Header
+	{
+		static constexpr const TCHAR* TransportId = TEXT("UE-MessageBus-TransportId");
+	}
+}
+
+class FWebSocketMessagingModule : public IModuleInterface
+{
+public:
+
+	/** IModuleInterface implementation */
+	virtual void StartupModule() override;
+	virtual void ShutdownModule() override;
+
+	virtual bool HandleSettingsSaved();
+
+	virtual void InitializeBridge();
+	virtual void ShutdownBridge();
+
+protected:
+	/** Holds the message bridge if present. */
+	TSharedPtr<class IMessageBridge, ESPMode::ThreadSafe> MessageBridge;
+};
diff --git a/Engine/Plugins/Experimental/WebSocketMessaging/Source/WebSocketMessaging/Private/WebSocketMessagingSettings.h b/Engine/Plugins/Experimental/WebSocketMessaging/Source/WebSocketMessaging/Private/WebSocketMessagingSettings.h
new file mode 100644
--- /dev/null
+++ b/Engine/Plugins/Experimental/WebSocketMessaging/Source/WebSocketMessaging/Private/WebSocketMessagingSettings.h
@@ -0,0 +1,44 @@
+// Copyright Epic Games, Inc. All Rights Reserved.
+
+#pragma once
+
+#include "Containers/Array.h"
+#include "Containers/Map.h"
+#include "Containers/UnrealString.h"
+#include "UObject/Object.h"
+#include "WebSocketMessagingSettings.generated.h"
+
+UENUM()
+enum class EWebSocketMessagingTransportFormat : uint8
+{
+	Json,
+	Cbor
+};
+
+UCLASS(config=Engine)
+class UWebSocketMessagingSettings : public UObject
+{
+	GENERATED_BODY()
+
+public:
+
+	/** Whether the WebSocket transport channel is enabled */
+	UPROPERTY(config, EditAnywhere, Category = Transport)
+	bool EnableTransport = false;
+
+	/** Bind the WebSocket server on the specified port (0 disables it) */
+	UPROPERTY(config, EditAnywhere, Category = Transport)
+	int32 ServerPort = 0;
+
+	/** Format used to serialize the messages on the server's WebSockets.*/
+	UPROPERTY(Config, EditAnywhere, Category = Transport)
+	EWebSocketMessagingTransportFormat ServerTransportFormat = EWebSocketMessagingTransportFormat::Cbor;
+	
+    /** The WebSocket Urls to connect to (Eg. ws://example.com/xyz) */
+    UPROPERTY(config, EditAnywhere, Category = Transport)
+    TArray<FString> ConnectToEndpoints;
+
+	/** Additional HTTP headers to set when connecting to endpoints */
+	UPROPERTY(config, EditAnywhere, Category = Transport)
+	TMap<FString, FString> HttpHeaders;
+};
diff --git a/Engine/Plugins/Experimental/WebSocketMessaging/Source/WebSocketMessaging/Public/WebSocketDeserializedMessage.h b/Engine/Plugins/Experimental/WebSocketMessaging/Source/WebSocketMessaging/Public/WebSocketDeserializedMessage.h
deleted file mode 100644
diff --git a/Engine/Plugins/Experimental/WebSocketMessaging/Source/WebSocketMessaging/Public/WebSocketMessageTransport.h b/Engine/Plugins/Experimental/WebSocketMessaging/Source/WebSocketMessaging/Public/WebSocketMessageTransport.h
deleted file mode 100644
diff --git a/Engine/Plugins/Experimental/WebSocketMessaging/Source/WebSocketMessaging/Public/WebSocketMessaging.h b/Engine/Plugins/Experimental/WebSocketMessaging/Source/WebSocketMessaging/Public/WebSocketMessaging.h
deleted file mode 100644
diff --git a/Engine/Plugins/Experimental/WebSocketMessaging/Source/WebSocketMessaging/Public/WebSocketMessagingSettings.h b/Engine/Plugins/Experimental/WebSocketMessaging/Source/WebSocketMessaging/Public/WebSocketMessagingSettings.h
deleted file mode 100644
diff --git a/Engine/Plugins/Experimental/WebSocketMessaging/Source/WebSocketMessaging/WebSocketMessaging.Build.cs b/Engine/Plugins/Experimental/WebSocketMessaging/Source/WebSocketMessaging/WebSocketMessaging.Build.cs
--- a/Engine/Plugins/Experimental/WebSocketMessaging/Source/WebSocketMessaging/WebSocketMessaging.Build.cs
+++ b/Engine/Plugins/Experimental/WebSocketMessaging/Source/WebSocketMessaging/WebSocketMessaging.Build.cs
@@ -8,18 +8,10 @@ public class WebSocketMessaging : ModuleRules
 	{
 		PCHUsage = ModuleRules.PCHUsageMode.UseExplicitOrSharedPCHs;
 
-
-		PublicDependencyModuleNames.AddRange(
-			new string[]
-			{
-				"Core",
-			}
-			);
-
-
 		PrivateDependencyModuleNames.AddRange(
 			new string[]
 			{
+				"Core",
 				"CoreUObject",
 				"Engine",
 				"Messaging",
@@ -30,6 +22,6 @@ public class WebSocketMessaging : ModuleRules
 				"Cbor",
 				"Serialization"
 			}
-			);
+		);
 	}
 }
