[WebSocket Messaging] Troubleshooting and backward compatibility improvements:

- Adding parse error to make troubleshooting easier when developping websocket clients.
- Supporting core redirects for message struct serialization. If message struct are moved or renamed, they can still be serialized correctly.

diff --git a/Engine/Plugins/Experimental/WebSocketMessaging/Source/WebSocketMessaging/Private/WebSocketDeserializedMessage.cpp b/Engine/Plugins/Experimental/WebSocketMessaging/Source/WebSocketMessaging/Private/WebSocketDeserializedMessage.cpp
--- a/Engine/Plugins/Experimental/WebSocketMessaging/Source/WebSocketMessaging/Private/WebSocketDeserializedMessage.cpp
+++ b/Engine/Plugins/Experimental/WebSocketMessaging/Source/WebSocketMessaging/Private/WebSocketDeserializedMessage.cpp
@@ -3,6 +3,7 @@
 
 #include "WebSocketDeserializedMessage.h"
 #include "JsonObjectConverter.h"
+#include "UObject/CoreRedirects.h"
 #include "WebSocketMessaging.h"
 
 FWebSocketDeserializedMessage::FWebSocketDeserializedMessage()
@@ -20,7 +21,7 @@ FWebSocketDeserializedMessage::~FWebSocketDeserializedMessage()
 	}
 }
 
-bool FWebSocketDeserializedMessage::ParseJson(const FString& Json)
+bool FWebSocketDeserializedMessage::ParseJson(const FString& Json, FString& OutParseError)
 {
 	static TMap<FString, EMessageScope> MessageScopeStringMapping =
 	{
@@ -44,17 +45,20 @@ bool FWebSocketDeserializedMessage::ParseJson(const FString& Json)
 		FString MessageType;
 		if (!RootObject->TryGetStringField(WebSocketMessaging::Tag::MessageType, MessageType))
 		{
+			OutParseError = FString::Printf(TEXT("Missing Mendatory Field: \"%s\"."), WebSocketMessaging::Tag::MessageType);
 			return false;
 		}
 
 		FString JsonSender;
 		if (!RootObject->TryGetStringField(WebSocketMessaging::Tag::Sender, JsonSender))
 		{
+			OutParseError = FString::Printf(TEXT("Missing Mendatory Field: \"%s\"."), WebSocketMessaging::Tag::Sender);
 			return false;
 		}
 
 		if (!FMessageAddress::Parse(JsonSender, Sender))
 		{
+			OutParseError = FString::Printf(TEXT("Field \"%s\": \"%s\" is not a valid Message Address."), WebSocketMessaging::Tag::Sender, *JsonSender);
 			return false;
 		}
 
@@ -70,12 +74,25 @@ bool FWebSocketDeserializedMessage::ParseJson(const FString& Json)
 		const TSharedPtr<FJsonObject>* JsonMessage = nullptr;
 		if (!RootObject->TryGetObjectField(WebSocketMessaging::Tag::Message, JsonMessage))
 		{
+			OutParseError = FString::Printf(TEXT("Missing Mendatory Field: \"%s\"."), WebSocketMessaging::Tag::Message);
 			return false;
 		}
 
 		UScriptStruct* ScriptStruct = FindObjectSafe<UScriptStruct>(nullptr, *MessageType);
+		
 		if (!ScriptStruct)
 		{
+			const FCoreRedirectObjectName OldObjectName(MessageType);
+			const FCoreRedirectObjectName NewObjectName = FCoreRedirects::GetRedirectedName(ECoreRedirectFlags::Type_Struct, OldObjectName);
+			if (NewObjectName.IsValid() && OldObjectName != NewObjectName)
+			{
+				ScriptStruct = FindObject<UScriptStruct>(nullptr, *NewObjectName.ToString());
+			}
+		}
+		
+		if (!ScriptStruct)
+		{
+			OutParseError = FString::Printf(TEXT("Field \"%s\": The message type \"%s\" is not a valid UScriptStruct."), WebSocketMessaging::Tag::MessageType, *MessageType);
 			return false;
 		}
 
@@ -93,6 +110,7 @@ bool FWebSocketDeserializedMessage::ParseJson(const FString& Json)
 		{
 			FMemory::Free(Message);
 			Message = nullptr;
+			OutParseError = FString::Printf(TEXT("Failed to deserialize UStruct \"%s\" from message data."), *MessageType);
 			return false;
 		}
 
@@ -117,6 +135,7 @@ bool FWebSocketDeserializedMessage::ParseJson(const FString& Json)
 			// unknown scope string
 			else
 			{
+				OutParseError = FString::Printf(TEXT("Field \"%s\": Unknown scope string: \"%s\"."), WebSocketMessaging::Tag::Scope, *ScopeString);
 				return false;
 			}
 		}
@@ -137,5 +156,6 @@ bool FWebSocketDeserializedMessage::ParseJson(const FString& Json)
 		return true;
 	}
 
+	OutParseError = TEXT("Message is not a valid json format.");
 	return false;
 }
diff --git a/Engine/Plugins/Experimental/WebSocketMessaging/Source/WebSocketMessaging/Private/WebSocketDeserializedMessage.h b/Engine/Plugins/Experimental/WebSocketMessaging/Source/WebSocketMessaging/Private/WebSocketDeserializedMessage.h
--- a/Engine/Plugins/Experimental/WebSocketMessaging/Source/WebSocketMessaging/Private/WebSocketDeserializedMessage.h
+++ b/Engine/Plugins/Experimental/WebSocketMessaging/Source/WebSocketMessaging/Private/WebSocketDeserializedMessage.h
@@ -18,7 +18,7 @@ public:
 
 	virtual ~FWebSocketDeserializedMessage() override;
 
-	virtual bool ParseJson(const FString& Json);
+	bool ParseJson(const FString& Json, FString& OutParseError);
 
 	virtual TSharedPtr<IMessageAttachment, ESPMode::ThreadSafe> GetAttachment() const override
 	{
diff --git a/Engine/Plugins/Experimental/WebSocketMessaging/Source/WebSocketMessaging/Private/WebSocketMessageTransport.cpp b/Engine/Plugins/Experimental/WebSocketMessaging/Source/WebSocketMessaging/Private/WebSocketMessageTransport.cpp
--- a/Engine/Plugins/Experimental/WebSocketMessaging/Source/WebSocketMessaging/Private/WebSocketMessageTransport.cpp
+++ b/Engine/Plugins/Experimental/WebSocketMessaging/Source/WebSocketMessaging/Private/WebSocketMessageTransport.cpp
@@ -142,14 +142,15 @@ void FWebSocketMessageTransport::OnConnectionError(const FString& Message, FWebS
 
 void FWebSocketMessageTransport::OnJsonMessage(const FString& Message, FWebSocketMessageConnectionRef WebSocketMessageConnection)
 {
-	TSharedRef<FWebSocketDeserializedMessage> Context = MakeShared<FWebSocketDeserializedMessage>();
-	if (Context->ParseJson(Message))
+	FString ParseError;
+	const TSharedRef<FWebSocketDeserializedMessage> Context = MakeShared<FWebSocketDeserializedMessage>();
+	if (Context->ParseJson(Message, ParseError))
 	{
 		TransportHandler->ReceiveTransportMessage(Context, WebSocketMessageConnection->Guid);
 	}
 	else
 	{
-		UE_LOG(LogWebSocketMessaging, Verbose, TEXT("Invalid Json Message received on %s"), *WebSocketMessageConnection->Url);
+		UE_LOG(LogWebSocketMessaging, Log, TEXT("Invalid Json Message received on %s: %s"), *WebSocketMessageConnection->Url, *ParseError);
 	}
 }
 
