[WebSocket Messaging] Service Discovery Beacon (and more)

- Adding a service discovery multicast beacon receiver to allow the websocket messaging service to be discovered directly by control apps. The advantage of having a discovery service is that the port and additional service information can be propagated.
- Adding command line option to specify the server port .The command line option is helpfull for launching a rundown server in the same project but with a different webocket port. The use case for that is having an editor running AND a package game, both having a rundown server on different ports.
- Adding a Message Serialization Json "Standardize Case" option. Since all my clients are already written, I'll leave the standardize case on to be compatible, but I wish I had found this option sooner.

diff --git a/Engine/Plugins/Experimental/WebSocketMessaging/Source/WebSocketMessaging/Private/WebSocketDeserializedMessage.cpp b/Engine/Plugins/Experimental/WebSocketMessaging/Source/WebSocketMessaging/Private/WebSocketDeserializedMessage.cpp
--- a/Engine/Plugins/Experimental/WebSocketMessaging/Source/WebSocketMessaging/Private/WebSocketDeserializedMessage.cpp
+++ b/Engine/Plugins/Experimental/WebSocketMessaging/Source/WebSocketMessaging/Private/WebSocketDeserializedMessage.cpp
@@ -4,7 +4,7 @@
 #include "WebSocketDeserializedMessage.h"
 #include "JsonObjectConverter.h"
 #include "UObject/CoreRedirects.h"
-#include "WebSocketMessaging.h"
+#include "WebSocketMessagingModule.h"
 
 FWebSocketDeserializedMessage::FWebSocketDeserializedMessage()
 	: Expiration(FDateTime::MaxValue())	// Make sure messages don't expire if no expiration is specified.
diff --git a/Engine/Plugins/Experimental/WebSocketMessaging/Source/WebSocketMessaging/Private/WebSocketMessageTransport.cpp b/Engine/Plugins/Experimental/WebSocketMessaging/Source/WebSocketMessaging/Private/WebSocketMessageTransport.cpp
--- a/Engine/Plugins/Experimental/WebSocketMessaging/Source/WebSocketMessaging/Private/WebSocketMessageTransport.cpp
+++ b/Engine/Plugins/Experimental/WebSocketMessaging/Source/WebSocketMessaging/Private/WebSocketMessageTransport.cpp
@@ -16,7 +16,7 @@
 #include "Serialization/ArrayWriter.h"
 #include "StructSerializer.h"
 #include "WebSocketDeserializedMessage.h"
-#include "WebSocketMessaging.h"
+#include "WebSocketMessagingModule.h"
 #include "WebSocketMessagingSettings.h"
 #include "WebSocketsModule.h"
 
@@ -46,13 +46,35 @@ FWebSocketMessageTransport::~FWebSocketMessageTransport()
 {
 }
 
+FName FWebSocketMessageTransport::GetDebugName() const
+{
+	static const FName DebugName("WebSocketMessageTransport");
+	return DebugName;
+}
+
 bool FWebSocketMessageTransport::StartTransport(IMessageTransportHandler& Handler)
 {
 	const UWebSocketMessagingSettings* Settings = GetDefault<UWebSocketMessagingSettings>();
-
+	
 	TransportHandler = &Handler;
+	
+	const int32 ServerPort = Settings->GetServerPort();
+	
+	// Cache the settings to be able to detect changes.
+	LastServerPort = ServerPort;
+	LastServerBindAddress = Settings->ServerBindAddress;
+	LastConnectionEndpoints = Settings->ConnectToEndpoints;
+	LastHttpHeaders = Settings->HttpHeaders;
 
-	if (Settings->ServerPort > 0)
+	FString ServerBindAddress = Settings->ServerBindAddress;
+
+	if (ServerBindAddress.Compare(TEXT("0.0.0.0")) == 0
+		|| ServerBindAddress.Compare(TEXT("any"), ESearchCase::IgnoreCase) == 0)
+	{
+		ServerBindAddress = TEXT("");	// Leaving empty will bind to all adapters.
+	}
+
+	if (ServerPort > 0)
 	{
 		IWebSocketNetworkingModule* WebSocketNetworkingModule = FModuleManager::Get().LoadModulePtr<IWebSocketNetworkingModule>(TEXT("WebSocketNetworking"));
 		if (WebSocketNetworkingModule)
@@ -60,14 +82,18 @@ bool FWebSocketMessageTransport::StartTransport(IMessageTransportHandler& Handle
 			Server = WebSocketNetworkingModule->CreateServer();
 			if (Server)
 			{
-				if (!Server->Init(Settings->ServerPort, FWebSocketClientConnectedCallBack::CreateThreadSafeSP(this, &FWebSocketMessageTransport::ClientConnected)))
+				FWebSocketClientConnectedCallBack Callback;
+				Callback.BindThreadSafeSP(this, &FWebSocketMessageTransport::ClientConnected);
+				
+				if (!Server->Init(ServerPort, Callback, ServerBindAddress))
 				{
-					UE_LOG(LogWebSocketMessaging, Log, TEXT("Unable to start WebSocketMessaging Server on port %d"), Settings->ServerPort);
+					Server.Reset();
+					UE_LOG(LogWebSocketMessaging, Error, TEXT("Unable to start WebSocketMessaging Server on port %d"), ServerPort);
 				}
 				else
 				{
 					ServerTickerHandle = FTSTicker::GetCoreTicker().AddTicker(FTickerDelegate::CreateThreadSafeSP(this, &FWebSocketMessageTransport::ServerTick));
-					UE_LOG(LogWebSocketMessaging, Log, TEXT("WebSocketMessaging Server started on port %d"), Settings->ServerPort);
+					UE_LOG(LogWebSocketMessaging, Log, TEXT("WebSocketMessaging Server started on port %d"), ServerPort);
 				}
 			}
 		}
@@ -100,7 +126,6 @@ bool FWebSocketMessageTransport::StartTransport(IMessageTransportHandler& Handle
 		WebSocketMessageConnections.Add(Guid, WebSocketMessageConnection);
 
 		WebSocketConnection->Connect();
-
 	}
 
 	return true;
@@ -108,6 +133,8 @@ bool FWebSocketMessageTransport::StartTransport(IMessageTransportHandler& Handle
 
 void FWebSocketMessageTransport::StopTransport()
 {
+	FTSTicker::GetCoreTicker().RemoveTicker(ServerTickerHandle);
+	
 	if (Server.IsValid())
 	{
 		Server.Reset();
@@ -121,6 +148,21 @@ void FWebSocketMessageTransport::StopTransport()
 	WebSocketMessageConnections.Empty();
 }
 
+bool FWebSocketMessageTransport::NeedsRestart() const
+{
+	const UWebSocketMessagingSettings* Settings = GetDefault<UWebSocketMessagingSettings>();
+	
+	if (LastServerPort != Settings->GetServerPort()
+		|| LastServerBindAddress != Settings->ServerBindAddress
+		|| LastConnectionEndpoints != Settings->ConnectToEndpoints
+		|| !LastHttpHeaders.OrderIndependentCompareEqual(Settings->HttpHeaders))
+	{
+		return true;
+	}
+
+	return false;
+}
+
 void FWebSocketMessageTransport::OnClosed(int32 Code, const FString& Reason, bool bUserClose, FWebSocketMessageConnectionRef WebSocketMessageConnection)
 {
 	UE_LOG(LogWebSocketMessaging, Log, TEXT("Connection to %s closed, Code: %d Reason: \"%s\" UserClose: %s, retrying..."), *WebSocketMessageConnection->Url, Code, *Reason, bUserClose ? TEXT("true") : TEXT("false"));
@@ -165,24 +207,24 @@ class FWebSocketMessageTransportSerializeHelper
 public:
 	static const TMap<EMessageScope, FString> MessageScopeStringMapping;
 
-	static bool Serialize(const TSharedRef<IMessageContext, ESPMode::ThreadSafe>& Context, FString& OutJsonMessage)
+	static bool Serialize(const TSharedRef<IMessageContext, ESPMode::ThreadSafe>& InContext, bool bInStandardizeCase, FString& OutJsonMessage)
 	{
 		TSharedRef<FJsonObject> JsonRoot = MakeShared<FJsonObject>();
-		JsonRoot->SetStringField(WebSocketMessaging::Tag::Sender, Context->GetSender().ToString());
+		JsonRoot->SetStringField(WebSocketMessaging::Tag::Sender, InContext->GetSender().ToString());
 		TArray<TSharedPtr<FJsonValue>> JsonRecipients;
-		for (const FMessageAddress& Recipient : Context->GetRecipients())
+		for (const FMessageAddress& Recipient : InContext->GetRecipients())
 		{
 			JsonRecipients.Add(MakeShared<FJsonValueString>(Recipient.ToString()));
 		}
 		JsonRoot->SetArrayField(WebSocketMessaging::Tag::Recipients, JsonRecipients);
-		JsonRoot->SetStringField(WebSocketMessaging::Tag::MessageType, Context->GetMessageTypePathName().ToString());
-		JsonRoot->SetNumberField(WebSocketMessaging::Tag::Expiration, Context->GetExpiration().ToUnixTimestamp());
-		JsonRoot->SetNumberField(WebSocketMessaging::Tag::TimeSent, Context->GetTimeSent().ToUnixTimestamp());
-		JsonRoot->SetStringField(WebSocketMessaging::Tag::Scope, MessageScopeStringMapping[Context->GetScope()]);
+		JsonRoot->SetStringField(WebSocketMessaging::Tag::MessageType, InContext->GetMessageTypePathName().ToString());
+		JsonRoot->SetNumberField(WebSocketMessaging::Tag::Expiration, InContext->GetExpiration().ToUnixTimestamp());
+		JsonRoot->SetNumberField(WebSocketMessaging::Tag::TimeSent, InContext->GetTimeSent().ToUnixTimestamp());
+		JsonRoot->SetStringField(WebSocketMessaging::Tag::Scope, MessageScopeStringMapping[InContext->GetScope()]);
 
 
 		TSharedRef<FJsonObject> JsonAnnotations = MakeShared<FJsonObject>();
-		for (const TPair<FName, FString>& Pair : Context->GetAnnotations())
+		for (const TPair<FName, FString>& Pair : InContext->GetAnnotations())
 		{
 			JsonAnnotations->SetStringField(Pair.Key.ToString(), Pair.Value);
 		}
@@ -190,8 +232,11 @@ public:
 
 		TSharedRef<FJsonObject> OutJsonObject = MakeShared<FJsonObject>();
 
-		// Remark: This will change the case of field names, see StandardizeCase.
-		if (!FJsonObjectConverter::UStructToJsonObject(Context->GetMessageTypeInfo().Get(), Context->GetMessage(), OutJsonObject))
+		constexpr int64 CheckFlags = 0;
+		constexpr int64 SkipFlags = 0;
+		const EJsonObjectConversionFlags ConversionFlags = bInStandardizeCase ? EJsonObjectConversionFlags::None : EJsonObjectConversionFlags::SkipStandardizeCase;
+		if (!FJsonObjectConverter::UStructToJsonObject(InContext->GetMessageTypeInfo().Get(), InContext->GetMessage(), OutJsonObject,
+			CheckFlags, SkipFlags, nullptr, ConversionFlags))
 		{
 			return false;
 		}
@@ -202,7 +247,7 @@ public:
 		return FJsonSerializer::Serialize(JsonRoot, Writer);
 	}
 
-	static bool Serialize(const TSharedRef<IMessageContext, ESPMode::ThreadSafe>& Context, FArrayWriter& OutCborBinaryWriter)
+	static bool Serialize(const TSharedRef<IMessageContext, ESPMode::ThreadSafe>& InContext, bool bInStandardizeCase, FArrayWriter& OutCborBinaryWriter)
 	{
 		FCborHeader Header(ECborCode::Map | ECborCode::Indefinite);
 		OutCborBinaryWriter << Header;
@@ -211,32 +256,32 @@ public:
 			FCborWriter CborWriter(&OutCborBinaryWriter);
 			
 			CborWriter.WriteValue(FString(WebSocketMessaging::Tag::Sender));
-			CborWriter.WriteValue(Context->GetSender().ToString());
+			CborWriter.WriteValue(InContext->GetSender().ToString());
 			
 			CborWriter.WriteValue(FString(WebSocketMessaging::Tag::Recipients));
 			CborWriter.WriteContainerStart(ECborCode::Array, -1);
-			for (const FMessageAddress& Recipient : Context->GetRecipients())
+			for (const FMessageAddress& Recipient : InContext->GetRecipients())
 			{
 				CborWriter.WriteValue(Recipient.ToString());
 			}
 			CborWriter.WriteContainerEnd();
 			
 			CborWriter.WriteValue(FString(WebSocketMessaging::Tag::MessageType));
-			CborWriter.WriteValue(Context->GetMessageTypePathName().ToString());
+			CborWriter.WriteValue(InContext->GetMessageTypePathName().ToString());
 			
 			CborWriter.WriteValue(FString(WebSocketMessaging::Tag::Expiration));
-			CborWriter.WriteValue(Context->GetExpiration().ToUnixTimestamp());
+			CborWriter.WriteValue(InContext->GetExpiration().ToUnixTimestamp());
 			
 			CborWriter.WriteValue(FString(WebSocketMessaging::Tag::TimeSent));
-			CborWriter.WriteValue(Context->GetTimeSent().ToUnixTimestamp());
+			CborWriter.WriteValue(InContext->GetTimeSent().ToUnixTimestamp());
 			
 			CborWriter.WriteValue(FString(WebSocketMessaging::Tag::Scope));
-			CborWriter.WriteValue(MessageScopeStringMapping[Context->GetScope()]);
+			CborWriter.WriteValue(MessageScopeStringMapping[InContext->GetScope()]);
 			
 			
 			CborWriter.WriteValue(FString(WebSocketMessaging::Tag::Annotations));
 			CborWriter.WriteContainerStart(ECborCode::Map, -1);
-			for (const TPair<FName, FString>& Annotation : Context->GetAnnotations())
+			for (const TPair<FName, FString>& Annotation : InContext->GetAnnotations())
 			{
 				CborWriter.WriteValue(Annotation.Key.ToString());
 				CborWriter.WriteValue(Annotation.Value);
@@ -246,7 +291,7 @@ public:
 			CborWriter.WriteValue(FString(WebSocketMessaging::Tag::Message));
 		}
 		FCborStructSerializerBackend Backend(OutCborBinaryWriter, EStructSerializerBackendFlags::Default);
-		FStructSerializer::Serialize(Context->GetMessage(), *Context->GetMessageTypeInfo().Get(), Backend);
+		FStructSerializer::Serialize(InContext->GetMessage(), *InContext->GetMessageTypeInfo().Get(), Backend);
 			
 		Header.Set(ECborCode::Break);
 		OutCborBinaryWriter << Header;
@@ -260,13 +305,14 @@ public:
 		OutputType OutputMessage;
 		bool bIsSerializeAttempted = false;
 		bool bIsSerialized = false;
+		bool bStandardizeCase = true;
 
-		bool SerializeOnDemand(const TSharedRef<IMessageContext, ESPMode::ThreadSafe>& Context)
+		bool SerializeOnDemand(const TSharedRef<IMessageContext, ESPMode::ThreadSafe>& InContext)
 		{
-			if(!bIsSerializeAttempted)
+			if (!bIsSerializeAttempted)
 			{
 				bIsSerializeAttempted = true;
-				bIsSerialized = FWebSocketMessageTransportSerializeHelper::Serialize(Context, OutputMessage);
+				bIsSerialized = FWebSocketMessageTransportSerializeHelper::Serialize(InContext, bStandardizeCase, OutputMessage);
 			}
 			return bIsSerialized;
 		}
@@ -311,42 +357,42 @@ bool FWebSocketMessageTransport::TransportMessage(const TSharedRef<IMessageConte
 		return false;
 	}
 
-	// Remark: Json serializer uses UStructToJsonObject, which is going to change
-	// the field name case (see StandardizeCase), which is going to cause a difference
-	// with the field names in Cbor (which doesn't change the case).
-
+	const UWebSocketMessagingSettings* Settings = GetDefault<UWebSocketMessagingSettings>();
+	
 	FWebSocketMessageTransportSerializeHelper::TOnDemandSerializer<FString> JsonSerializer;
+	JsonSerializer.bStandardizeCase = Settings->bMessageSerializationStandardizeCase;
+	
 	FWebSocketMessageTransportSerializeHelper::TOnDemandSerializer<FArrayWriter> CborSerializer;
 	
 	// Serialize the message on demand in the appropriate format for each peer connections.
-	for (const TPair<FGuid, FWebSocketMessageConnectionRef>& Pair : RecipientConnections)
+	for (const TPair<FGuid, FWebSocketMessageConnectionRef>& Connection : RecipientConnections)
 	{
-		if (Pair.Value->WebSocketConnection.IsValid())
+		if (Connection.Value->WebSocketConnection.IsValid())
 		{
 			// Remark: client connections are always text/json
-			if(JsonSerializer.SerializeOnDemand(Context))
+			if (JsonSerializer.SerializeOnDemand(Context))
 			{
-				Pair.Value->WebSocketConnection->Send(JsonSerializer.OutputMessage);
+				Connection.Value->WebSocketConnection->Send(JsonSerializer.OutputMessage);
 			}
 		}
-		else if (Pair.Value->WebSocketServerConnection)
+		else if (Connection.Value->WebSocketServerConnection)
 		{
 			// Remark: server connections are always binary.
-			const UWebSocketMessagingSettings* Settings = GetDefault<UWebSocketMessagingSettings>();
-
-			if(Settings->ServerTransportFormat == EWebSocketMessagingTransportFormat::Json)
+			if (Settings->ServerTransportFormat == EWebSocketMessagingTransportFormat::Json)
 			{
-				if(JsonSerializer.SerializeOnDemand(Context))
+				if (JsonSerializer.SerializeOnDemand(Context))
 				{
-					FTCHARToUTF8 Converted(*JsonSerializer.OutputMessage);
-					Pair.Value->WebSocketServerConnection->Send(reinterpret_cast<const uint8*>(Converted.Get()), Converted.Length(), false);
+					auto MessageUtf8 = StringCast<UTF8CHAR>(*JsonSerializer.OutputMessage);
+					Connection.Value->WebSocketServerConnection->Send(
+						reinterpret_cast<const uint8*>(MessageUtf8.Get()), MessageUtf8.Length(), /*bPrependSize*/ false);
 				}
 			}
 			else
 			{
-				if(CborSerializer.SerializeOnDemand(Context))
+				if (CborSerializer.SerializeOnDemand(Context))
 				{
-					Pair.Value->WebSocketServerConnection->Send(CborSerializer.OutputMessage.GetData(), CborSerializer.OutputMessage.Num(), false);
+					Connection.Value->WebSocketServerConnection->Send(
+						CborSerializer.OutputMessage.GetData(), CborSerializer.OutputMessage.Num(), /*bPrependSize*/ false);
 				}
 			}
 		}
@@ -414,7 +460,7 @@ bool FWebSocketMessageTransport::ServerTick(float DeltaTime)
 
 void FWebSocketMessageTransport::ForgetTransportNode(FWebSocketMessageConnectionRef WebSocketMessageConnection)
 {
-	if(TransportHandler)
+	if (TransportHandler)
 	{
 		TransportHandler->ForgetTransportNode(WebSocketMessageConnection->Guid);
 	}
diff --git a/Engine/Plugins/Experimental/WebSocketMessaging/Source/WebSocketMessaging/Private/WebSocketMessageTransport.h b/Engine/Plugins/Experimental/WebSocketMessaging/Source/WebSocketMessaging/Private/WebSocketMessageTransport.h
--- a/Engine/Plugins/Experimental/WebSocketMessaging/Source/WebSocketMessaging/Private/WebSocketMessageTransport.h
+++ b/Engine/Plugins/Experimental/WebSocketMessaging/Source/WebSocketMessaging/Private/WebSocketMessageTransport.h
@@ -67,16 +67,14 @@ public:
 	FWebSocketMessageTransport();
 	virtual ~FWebSocketMessageTransport() override;
 
-	virtual FName GetDebugName() const override
-	{
-		return "WebSocketMessageTransport";
-	}
-
+	//~ Begin IMessageTransport
+	virtual FName GetDebugName() const override;
 	virtual bool StartTransport(IMessageTransportHandler& Handler) override;
-
 	virtual void StopTransport() override;
-
 	virtual bool TransportMessage(const TSharedRef<IMessageContext, ESPMode::ThreadSafe>& Context, const TArray<FGuid>& Recipients) override;
+	//~ End IMessageTransport
+
+	bool NeedsRestart() const;
 
 	virtual void OnJsonMessage(const FString& Message, FWebSocketMessageConnectionRef WebSocketMessageConnection);
 	virtual void OnConnected(FWebSocketMessageConnectionRef WebSocketMessageConnection);
@@ -95,9 +93,14 @@ public:
 protected:
 	void ForgetTransportNode(FWebSocketMessageConnectionRef WebSocketMessageConnection);
 	
-	IMessageTransportHandler* TransportHandler;
+	IMessageTransportHandler* TransportHandler = nullptr;
 	TMap<FGuid, FWebSocketMessageConnectionRef> WebSocketMessageConnections;
 
 	TUniquePtr<class IWebSocketServer> Server;
 	FTSTicker::FDelegateHandle ServerTickerHandle;
+
+	FString LastServerBindAddress;
+	int32 LastServerPort = INDEX_NONE;
+	TArray<FString> LastConnectionEndpoints;
+	TMap<FString, FString> LastHttpHeaders;
 };
diff --git a/Engine/Plugins/Experimental/WebSocketMessaging/Source/WebSocketMessaging/Private/WebSocketMessaging.cpp b/Engine/Plugins/Experimental/WebSocketMessaging/Source/WebSocketMessaging/Private/WebSocketMessaging.cpp
deleted file mode 100644
diff --git a/Engine/Plugins/Experimental/WebSocketMessaging/Source/WebSocketMessaging/Private/WebSocketMessaging.h b/Engine/Plugins/Experimental/WebSocketMessaging/Source/WebSocketMessaging/Private/WebSocketMessaging.h
deleted file mode 100644
diff --git a/Engine/Plugins/Experimental/WebSocketMessaging/Source/WebSocketMessaging/Private/WebSocketMessagingBeaconReceiver.cpp b/Engine/Plugins/Experimental/WebSocketMessaging/Source/WebSocketMessaging/Private/WebSocketMessagingBeaconReceiver.cpp
new file mode 100644
--- /dev/null
+++ b/Engine/Plugins/Experimental/WebSocketMessaging/Source/WebSocketMessaging/Private/WebSocketMessagingBeaconReceiver.cpp
@@ -0,0 +1,137 @@
+// Copyright Epic Games, Inc. All Rights Reserved.
+
+#include "WebSocketMessagingBeaconReceiver.h"
+
+#include "Backends/CborStructSerializerBackend.h"
+#include "CborWriter.h"
+#include "Interfaces/IPv4/IPv4Address.h"
+#include "JsonObjectConverter.h"
+#include "Serialization/ArrayReader.h"
+#include "Serialization/ArrayWriter.h"
+#include "StructSerializer.h"
+#include "WebSocketMessagingModule.h"
+#include "WebSocketMessagingSettings.h"
+
+namespace UE::WebSocketMessaging::BeaconReceiver::Private
+{
+	constexpr uint8 ProtocolVersion = 0;
+	const TArray<uint8> ProtocolIdentifier = { 'U', 'E', '_', 'W', 'S', 'M', 'B'};
+
+	bool SerializeToJson(const FWebSocketMessagingBeaconPayload& InPayload, FArrayWriter& OutResponseData)
+	{
+		const UWebSocketMessagingSettings* Settings = GetDefault<UWebSocketMessagingSettings>();
+		constexpr int64 CheckFlags = 0;
+		constexpr int64 SkipFlags = 0;
+		const EJsonObjectConversionFlags ConversionFlags = Settings->bMessageSerializationStandardizeCase ?
+			EJsonObjectConversionFlags::None : EJsonObjectConversionFlags::SkipStandardizeCase;
+
+		const TSharedRef<FJsonObject> JsonObject = MakeShared<FJsonObject>();
+		
+		if (!FJsonObjectConverter::UStructToJsonObject(FWebSocketMessagingBeaconPayload::StaticStruct(), &InPayload, JsonObject,
+			CheckFlags, SkipFlags, nullptr, ConversionFlags))
+		{
+			return false;
+		}
+
+		// Todo: see if there is a way to serialize the json object with an utf8 writer directly. 
+		FString JsonString;
+		const TSharedRef<TJsonWriter<>> Writer = TJsonWriterFactory<>::Create(&JsonString);
+		if (!FJsonSerializer::Serialize(JsonObject, Writer))
+		{
+			return false;
+		}
+
+		const auto MessageUtf8 = StringCast<UTF8CHAR>(*JsonString);
+		OutResponseData.Serialize(const_cast<UTF8CHAR *>(MessageUtf8.Get()), MessageUtf8.Length());
+		return true;
+	}
+
+	bool SerializeToCbor(const FWebSocketMessagingBeaconPayload& InPayload, FArrayWriter& OutResponseData)
+	{
+		FCborWriter CborWriter(&OutResponseData);
+		FCborStructSerializerBackend Backend(OutResponseData, EStructSerializerBackendFlags::Default);
+		FStructSerializer::Serialize(&InPayload, *FWebSocketMessagingBeaconPayload::StaticStruct(), Backend);
+		return true;
+	}
+
+}
+
+FWebSocketMessagingBeaconReceiver::FWebSocketMessagingBeaconReceiver()
+	: FDiscoveryBeaconReceiver(
+		TEXT("WebSocketMessagingBeaconReceiver"),
+		UE::WebSocketMessaging::BeaconReceiver::Private::ProtocolIdentifier,
+		UE::WebSocketMessaging::BeaconReceiver::Private::ProtocolVersion
+	)
+{
+}
+
+bool FWebSocketMessagingBeaconReceiver::NeedsRestart() const
+{
+	const UWebSocketMessagingSettings* Settings = GetDefault<UWebSocketMessagingSettings>();
+	if (!Settings)
+	{
+		return false;
+	}
+	
+	if (LastDiscoveryEndpoint != Settings->DiscoveryEndpoint
+		|| LastDiscoveryPort != Settings->DiscoveryPort)
+	{
+		return true;
+	}
+
+	return false;
+}
+
+void FWebSocketMessagingBeaconReceiver::Startup()
+{
+	const UWebSocketMessagingSettings& Settings = *GetDefault<UWebSocketMessagingSettings>();
+	LastDiscoveryEndpoint = Settings.DiscoveryEndpoint;
+	LastDiscoveryPort = Settings.DiscoveryPort;
+	
+	FDiscoveryBeaconReceiver::Startup();
+}
+
+bool FWebSocketMessagingBeaconReceiver::GetDiscoveryAddress(FIPv4Address& OutAddress) const
+{
+	const UWebSocketMessagingSettings& Settings = *GetDefault<UWebSocketMessagingSettings>();
+	FIPv4Address DiscoveryAddress;
+	if (!FIPv4Address::Parse(Settings.DiscoveryEndpoint, DiscoveryAddress))
+	{
+		UE_LOG(LogWebSocketMessaging, Error, TEXT("Failed to parse WebSocket Messaging discovery endpoint address \"%s\""), *Settings.DiscoveryEndpoint);
+		return false;
+	}
+
+	OutAddress = DiscoveryAddress;
+	return true;
+}
+
+int32 FWebSocketMessagingBeaconReceiver::GetDiscoveryPort() const
+{
+	const UWebSocketMessagingSettings* Settings = GetDefault<UWebSocketMessagingSettings>();
+	return Settings ? Settings->DiscoveryPort : 0;
+}
+
+bool FWebSocketMessagingBeaconReceiver::MakeBeaconResponse(uint8 InBeaconProtocolVersion, FArrayReader& InMessageData, FArrayWriter& OutResponseData) const
+{
+	const UWebSocketMessagingSettings* Settings = GetDefault<UWebSocketMessagingSettings>();
+
+	FWebSocketMessagingBeaconPayload Reply;
+
+	if (Settings && Settings->EnableTransport)
+	{
+		Reply.Services.Reserve(1);
+		Reply.Services.Add({FString(TEXT("WebSocketMessaging")), Settings->GetServerPort()});
+	}
+
+	using namespace UE::WebSocketMessaging::BeaconReceiver::Private;
+
+	// Json is the default format.
+	const EWebSocketMessagingTransportFormat PayloadFormat = Settings ? Settings->DiscoveryPayloadFormat : EWebSocketMessagingTransportFormat::Json; 
+
+	if (PayloadFormat == EWebSocketMessagingTransportFormat::Json)
+	{
+		return SerializeToJson(Reply, OutResponseData);
+	}
+
+	return SerializeToCbor(Reply, OutResponseData);
+}
diff --git a/Engine/Plugins/Experimental/WebSocketMessaging/Source/WebSocketMessaging/Private/WebSocketMessagingBeaconReceiver.h b/Engine/Plugins/Experimental/WebSocketMessaging/Source/WebSocketMessaging/Private/WebSocketMessagingBeaconReceiver.h
new file mode 100644
--- /dev/null
+++ b/Engine/Plugins/Experimental/WebSocketMessaging/Source/WebSocketMessaging/Private/WebSocketMessagingBeaconReceiver.h
@@ -0,0 +1,56 @@
+// Copyright Epic Games, Inc. All Rights Reserved.
+
+#pragma once
+
+#include "DiscoveryBeaconReceiver.h"
+
+#include "WebSocketMessagingBeaconReceiver.generated.h"
+
+/** Service entry */
+USTRUCT()
+struct FWebSocketMessagingBeaconService
+{
+	GENERATED_BODY()
+
+	UPROPERTY()
+	FString Name;
+
+	UPROPERTY()
+	int32 Port = 0;
+};
+
+/** Payload of the beacon message used as serialization helper. */
+USTRUCT()
+struct FWebSocketMessagingBeaconPayload
+{
+	GENERATED_BODY()
+
+	UPROPERTY()
+	TArray<FWebSocketMessagingBeaconService> Services;
+};
+
+/**
+ * Receives beacon messages from the External Apps and replies with connection information.
+ * This allows the apps to detect compatible Unreal instances on the local network and list them for the user.
+ */
+class FWebSocketMessagingBeaconReceiver : public FDiscoveryBeaconReceiver
+{
+public:
+	FWebSocketMessagingBeaconReceiver();
+	virtual ~FWebSocketMessagingBeaconReceiver() override = default;
+
+	bool NeedsRestart() const;
+	
+	//~ Begin FDiscoveryBeaconReceiver
+	virtual void Startup() override;
+
+protected:
+	virtual bool GetDiscoveryAddress(FIPv4Address& OutAddress) const override;
+	virtual int32 GetDiscoveryPort() const override;
+	virtual bool MakeBeaconResponse(uint8 InBeaconProtocolVersion, FArrayReader& InMessageData, FArrayWriter& OutResponseData) const override;
+	//~ End FDiscoveryBeaconReceiver
+
+private:
+	FString LastDiscoveryEndpoint;
+	int32 LastDiscoveryPort;
+};
\ No newline at end of file
diff --git a/Engine/Plugins/Experimental/WebSocketMessaging/Source/WebSocketMessaging/Private/WebSocketMessagingModule.cpp b/Engine/Plugins/Experimental/WebSocketMessaging/Source/WebSocketMessaging/Private/WebSocketMessagingModule.cpp
new file mode 100644
--- /dev/null
+++ b/Engine/Plugins/Experimental/WebSocketMessaging/Source/WebSocketMessaging/Private/WebSocketMessagingModule.cpp
@@ -0,0 +1,145 @@
+// Copyright Epic Games, Inc. All Rights Reserved.
+
+#include "WebSocketMessagingModule.h"
+#include "IMessageBridge.h"
+#if WITH_EDITOR
+#include "ISettingsModule.h"
+#include "ISettingsSection.h"
+#endif
+#include "MessageBridgeBuilder.h"
+#include "WebSocketMessageTransport.h"
+#include "WebSocketMessagingBeaconReceiver.h"
+#include "WebSocketMessagingSettings.h"
+
+DEFINE_LOG_CATEGORY(LogWebSocketMessaging);
+
+#define LOCTEXT_NAMESPACE "FWebSocketMessagingModule"
+
+void FWebSocketMessagingModule::StartupModule()
+{
+	if (!FModuleManager::Get().LoadModule(TEXT("WebSockets")))
+	{
+		return;
+	}
+#if WITH_EDITOR
+	// register settings
+	ISettingsModule* SettingsModule = FModuleManager::GetModulePtr<ISettingsModule>("Settings");
+
+	if (SettingsModule != nullptr)
+	{
+		ISettingsSectionPtr SettingsSection = SettingsModule->RegisterSettings("Project", "Plugins", "WebSocketMessaging",
+			LOCTEXT("WebSocketMessagingSettingsName", "WebSocket Messaging"),
+			LOCTEXT("WebSocketMessagingSettingsDescription", "Configure the WebSocket Messaging plug-in."),
+			GetMutableDefault<UWebSocketMessagingSettings>()
+		);
+
+		if (SettingsSection.IsValid())
+		{
+			SettingsSection->OnModified().BindRaw(this, &FWebSocketMessagingModule::HandleSettingsSaved);
+		}
+	}
+#endif // WITH_EDITOR
+
+	// trigger initialization
+	HandleSettingsSaved();
+}
+
+void FWebSocketMessagingModule::ShutdownModule()
+{
+#if WITH_EDITOR
+	// unregister settings
+	ISettingsModule* SettingsModule = FModuleManager::GetModulePtr<ISettingsModule>("Settings");
+
+	if (SettingsModule != nullptr)
+	{
+		SettingsModule->UnregisterSettings("Project", "Plugins", "WebSocketMessaging");
+	}
+#endif
+}
+
+bool FWebSocketMessagingModule::IsTransportRunning() const
+{
+	return MessageBridge.IsValid();
+}
+
+int32 FWebSocketMessagingModule::GetServerPort() const
+{
+	if (const UWebSocketMessagingSettings* Settings = GetDefault<UWebSocketMessagingSettings>())
+	{
+		return Settings->GetServerPort();
+	}
+	return 0;
+}
+
+bool FWebSocketMessagingModule::HandleSettingsSaved()
+{
+	if (GetDefault<UWebSocketMessagingSettings>()->EnableTransport)
+	{
+		InitializeBridge();
+	}
+	else
+	{
+		ShutdownBridge();
+	}
+
+	if (GetDefault<UWebSocketMessagingSettings>()->bEnableDiscoveryListener)
+	{
+		InitializeBeaconReceiver();
+	}
+	else
+	{
+		ShutdownBeaconReceiver();
+	}
+	return true;
+}
+
+void FWebSocketMessagingModule::InitializeBridge()
+{
+	const TSharedPtr<FWebSocketMessageTransport, ESPMode::ThreadSafe> Transport = TransportWeak.Pin();
+
+	if (Transport && !Transport->NeedsRestart())
+	{
+		return;
+	}
+
+	ShutdownBridge();
+
+	const TSharedRef<FWebSocketMessageTransport, ESPMode::ThreadSafe> NewTransport = MakeShared<FWebSocketMessageTransport, ESPMode::ThreadSafe>();
+	TransportWeak = NewTransport;
+	MessageBridge = FMessageBridgeBuilder().UsingTransport(NewTransport).Build();
+}
+
+void FWebSocketMessagingModule::ShutdownBridge()
+{
+	if (MessageBridge.IsValid())
+	{
+		MessageBridge->Disable();
+		FPlatformProcess::Sleep(0.1f);
+		MessageBridge.Reset();
+	}
+}
+
+void FWebSocketMessagingModule::InitializeBeaconReceiver()
+{
+	if (BeaconReceiver && !BeaconReceiver->NeedsRestart())
+	{
+		return;
+	}
+	
+	ShutdownBeaconReceiver();
+	BeaconReceiver = MakeUnique<FWebSocketMessagingBeaconReceiver>();
+	BeaconReceiver->Startup();
+}
+
+void FWebSocketMessagingModule::ShutdownBeaconReceiver()
+{
+	if (BeaconReceiver)
+	{
+		BeaconReceiver->Shutdown();
+	}
+	BeaconReceiver.Reset();
+}
+
+#undef LOCTEXT_NAMESPACE
+
+IMPLEMENT_MODULE(FWebSocketMessagingModule, WebSocketMessaging)
\ No newline at end of file
diff --git a/Engine/Plugins/Experimental/WebSocketMessaging/Source/WebSocketMessaging/Private/WebSocketMessagingModule.h b/Engine/Plugins/Experimental/WebSocketMessaging/Source/WebSocketMessaging/Private/WebSocketMessagingModule.h
new file mode 100644
--- /dev/null
+++ b/Engine/Plugins/Experimental/WebSocketMessaging/Source/WebSocketMessaging/Private/WebSocketMessagingModule.h
@@ -0,0 +1,65 @@
+// Copyright Epic Games, Inc. All Rights Reserved.
+
+#pragma once
+
+#include "IWebSocketMessagingModule.h"
+#include "Logging/LogMacros.h"
+#include "Modules/ModuleManager.h"
+#include "Templates/SharedPointer.h"
+
+class FWebSocketMessageTransport;
+class FWebSocketMessagingBeaconReceiver;
+class UWebSocketMessagingSettings;
+
+DECLARE_LOG_CATEGORY_EXTERN(LogWebSocketMessaging, Log, All);
+
+namespace WebSocketMessaging
+{
+	namespace Tag
+	{
+		static constexpr const TCHAR* Sender = TEXT("Sender");
+		static constexpr const TCHAR* Recipients = TEXT("Recipients");
+		static constexpr const TCHAR* Expiration = TEXT("Expiration");
+		static constexpr const TCHAR* TimeSent = TEXT("TimeSent");
+		static constexpr const TCHAR* Annotations = TEXT("Annotations");
+		static constexpr const TCHAR* Scope = TEXT("Scope");
+		static constexpr const TCHAR* MessageType = TEXT("MessageType");
+		static constexpr const TCHAR* Message = TEXT("Message");
+	}
+
+	namespace Header
+	{
+		static constexpr const TCHAR* TransportId = TEXT("UE-MessageBus-TransportId");
+	}
+}
+
+// Todo: implement INetworkMessagingExtension to expose more service controls.
+class FWebSocketMessagingModule : public IWebSocketMessagingModule
+{
+public:
+	//~ Begin IModuleInterface
+	virtual void StartupModule() override;
+	virtual void ShutdownModule() override;
+	//~ End IModuleInterface
+	
+	//~ Begin IWebSocketMessaging
+	virtual bool IsTransportRunning() const override;
+	virtual int32 GetServerPort() const override;
+	//~ End IWebSocketMessaging
+	
+	virtual bool HandleSettingsSaved();
+	
+	virtual void InitializeBridge();
+	virtual void ShutdownBridge();
+
+	void InitializeBeaconReceiver();
+	void ShutdownBeaconReceiver();
+
+protected:
+	/** Holds the message bridge if present. */
+	TSharedPtr<class IMessageBridge, ESPMode::ThreadSafe> MessageBridge;
+	/** Keep track of the transport for access to derived functions. */
+	TWeakPtr<FWebSocketMessageTransport, ESPMode::ThreadSafe> TransportWeak;
+	/** Multicast Discovery Beacon Receiver, if present. */
+	TUniquePtr<FWebSocketMessagingBeaconReceiver> BeaconReceiver;
+};
\ No newline at end of file
diff --git a/Engine/Plugins/Experimental/WebSocketMessaging/Source/WebSocketMessaging/Private/WebSocketMessagingSettings.cpp b/Engine/Plugins/Experimental/WebSocketMessaging/Source/WebSocketMessaging/Private/WebSocketMessagingSettings.cpp
new file mode 100644
--- /dev/null
+++ b/Engine/Plugins/Experimental/WebSocketMessaging/Source/WebSocketMessaging/Private/WebSocketMessagingSettings.cpp
@@ -0,0 +1,17 @@
+// Copyright Epic Games, Inc. All Rights Reserved.
+
+#include "WebSocketMessagingSettings.h"
+
+#include "Misc/CommandLine.h"
+#include "Misc/Parse.h"
+
+int32 UWebSocketMessagingSettings::GetServerPort() const
+{
+	int32 ServerPortCmdLine;
+	// Allow port to be specified on the command line for local server in game mode.
+	if (FParse::Value(FCommandLine::Get(), TEXT("-WebSocketMessagingServerPort="), ServerPortCmdLine))
+	{
+		return ServerPortCmdLine;
+	}
+	return ServerPort;
+}
diff --git a/Engine/Plugins/Experimental/WebSocketMessaging/Source/WebSocketMessaging/Private/WebSocketMessagingSettings.h b/Engine/Plugins/Experimental/WebSocketMessaging/Source/WebSocketMessaging/Private/WebSocketMessagingSettings.h
--- a/Engine/Plugins/Experimental/WebSocketMessaging/Source/WebSocketMessaging/Private/WebSocketMessagingSettings.h
+++ b/Engine/Plugins/Experimental/WebSocketMessaging/Source/WebSocketMessaging/Private/WebSocketMessagingSettings.h
@@ -21,24 +21,75 @@ class UWebSocketMessagingSettings : public UObject
 	GENERATED_BODY()
 
 public:
+	int32 GetServerPort() const;
 
 	/** Whether the WebSocket transport channel is enabled */
-	UPROPERTY(config, EditAnywhere, Category = Transport)
+	UPROPERTY(config, EditAnywhere, Category=Transport)
 	bool EnableTransport = false;
 
-	/** Bind the WebSocket server on the specified port (0 disables it) */
-	UPROPERTY(config, EditAnywhere, Category = Transport)
+	/**
+	 * Bind the WebSocket server on the specified port (0 disables it)
+	 * Can be specified on the command line with `-WebSocketMessagingServerPort=`
+	 */
+	UPROPERTY(config, EditAnywhere, Category=Transport)
 	int32 ServerPort = 0;
 
+	/**
+	 * The address to bind the websocket server to.
+	 * 0.0.0.0 will open the connection to everyone on your network,
+	 * while 127.0.0.1 will only allow local requests to come through. 
+	 */
+	UPROPERTY(config, EditAnywhere, Category=Transport)
+	FString ServerBindAddress = TEXT("0.0.0.0");
+
 	/** Format used to serialize the messages on the server's WebSockets.*/
-	UPROPERTY(Config, EditAnywhere, Category = Transport)
+	UPROPERTY(Config, EditAnywhere, Category=Transport)
 	EWebSocketMessagingTransportFormat ServerTransportFormat = EWebSocketMessagingTransportFormat::Cbor;
+
+	/**
+	 * For Json formatting only:
+	 * If enabled, the "message" part of the messages will have "standardized case" (see FJsonObjectConverter::StandardizeCase).
+	 * Mainly, the first character of the field name will be lower case.
+	 * If disabled, the field names are not modified.
+	 *
+	 * For Cbor formatting, the field names are not modified.
+	 */ 
+	UPROPERTY(Config, EditAnywhere, Category=Transport, meta=(EditCondition="ServerTransportFormat == EWebSocketMessagingTransportFormat::Json"))
+	bool bMessageSerializationStandardizeCase = true;
 	
     /** The WebSocket Urls to connect to (Eg. ws://example.com/xyz) */
-    UPROPERTY(config, EditAnywhere, Category = Transport)
+    UPROPERTY(config, EditAnywhere, Category=Transport)
     TArray<FString> ConnectToEndpoints;
 
 	/** Additional HTTP headers to set when connecting to endpoints */
-	UPROPERTY(config, EditAnywhere, Category = Transport)
+	UPROPERTY(config, EditAnywhere, Category=Transport)
 	TMap<FString, FString> HttpHeaders;
+
+	/**
+	 * Enables the Multicast Service Discovery
+	 */
+	UPROPERTY(config, EditAnywhere, Category=Discovery)
+	bool bEnableDiscoveryListener = false;
+
+	/**
+	 * The IP endpoint to listen to for multicast discovery messages.
+	 * The multicast IP address must be in the range 224.0.0.0 to 239.255.255.255.
+	 */
+	UPROPERTY(config, EditAnywhere, Category=Discovery, meta=(EditCondition="bEnableDiscoveryListener"))
+	FString DiscoveryEndpoint = "230.0.0.4";
+
+	/**
+	 * The port to listen to for app discovery messages.
+	 */
+	UPROPERTY(config, EditAnywhere, Category=Discovery, meta=(EditCondition="bEnableDiscoveryListener"))
+	int32 DiscoveryPort = 6667;
+
+	/**
+	 * Format used to serialize the discovery beacon payload.
+	 * The discovery beacon response will have a header followed by the payload.
+	 * Header Format: 1 byte Protocol Version, 16 bytes Beacon Guid
+	 * The payload is an array of "services" with name and port per service. 
+	 */
+	UPROPERTY(Config, EditAnywhere, Category=Discovery, meta=(EditCondition="bEnableDiscoveryListener"))
+	EWebSocketMessagingTransportFormat DiscoveryPayloadFormat = EWebSocketMessagingTransportFormat::Json;
 };
diff --git a/Engine/Plugins/Experimental/WebSocketMessaging/Source/WebSocketMessaging/Public/IWebSocketMessagingModule.h b/Engine/Plugins/Experimental/WebSocketMessaging/Source/WebSocketMessaging/Public/IWebSocketMessagingModule.h
new file mode 100644
--- /dev/null
+++ b/Engine/Plugins/Experimental/WebSocketMessaging/Source/WebSocketMessaging/Public/IWebSocketMessagingModule.h
@@ -0,0 +1,18 @@
+// Copyright Epic Games, Inc. All Rights Reserved.
+
+#pragma once
+
+#include "Modules/ModuleInterface.h"
+
+/**
+ * Interface to provide information on the web socket messaging transport.
+ */
+class IWebSocketMessagingModule : public IModuleInterface
+{
+public:
+	/** Returns true if the transport is currently running. */
+	virtual bool IsTransportRunning() const = 0;
+	
+	/** Returns the current server port, either from settings or command line. */
+	virtual int32 GetServerPort() const = 0;
+};
\ No newline at end of file
diff --git a/Engine/Plugins/Experimental/WebSocketMessaging/Source/WebSocketMessaging/WebSocketMessaging.Build.cs b/Engine/Plugins/Experimental/WebSocketMessaging/Source/WebSocketMessaging/WebSocketMessaging.Build.cs
--- a/Engine/Plugins/Experimental/WebSocketMessaging/Source/WebSocketMessaging/WebSocketMessaging.Build.cs
+++ b/Engine/Plugins/Experimental/WebSocketMessaging/Source/WebSocketMessaging/WebSocketMessaging.Build.cs
@@ -13,8 +13,10 @@ public class WebSocketMessaging : ModuleRules
 			{
 				"Core",
 				"CoreUObject",
+				"DiscoveryBeaconReceiver",
 				"Engine",
 				"Messaging",
+				"Networking",
 				"WebSockets",
 				"WebSocketNetworking",
 				"Json",
diff --git a/Engine/Plugins/Experimental/WebSocketMessaging/WebSocketMessaging.uplugin b/Engine/Plugins/Experimental/WebSocketMessaging/WebSocketMessaging.uplugin
--- a/Engine/Plugins/Experimental/WebSocketMessaging/WebSocketMessaging.uplugin
+++ b/Engine/Plugins/Experimental/WebSocketMessaging/WebSocketMessaging.uplugin
@@ -27,6 +27,10 @@
 		}
 	],
 	"Plugins": [
+		{
+			"Name": "DiscoveryBeaconReceiver",
+			"Enabled": true
+		},
 		{
 			"Name": "WebSocketNetworking",
 			"Enabled": true
