[WebSocket Messaging] Code Style clean up.

diff --git a/Engine/Plugins/Experimental/WebSocketMessaging/Source/WebSocketMessaging/Private/WebSocketDeserializedMessage.cpp b/Engine/Plugins/Experimental/WebSocketMessaging/Source/WebSocketMessaging/Private/WebSocketDeserializedMessage.cpp
--- a/Engine/Plugins/Experimental/WebSocketMessaging/Source/WebSocketMessaging/Private/WebSocketDeserializedMessage.cpp
+++ b/Engine/Plugins/Experimental/WebSocketMessaging/Source/WebSocketMessaging/Private/WebSocketDeserializedMessage.cpp
@@ -21,7 +21,7 @@ FWebSocketDeserializedMessage::~FWebSocketDeserializedMessage()
 	}
 }
 
-bool FWebSocketDeserializedMessage::ParseJson(const FString& Json, FString& OutParseError)
+bool FWebSocketDeserializedMessage::ParseJson(const FString& InJson, FString& OutParseError)
 {
 	static TMap<FString, EMessageScope> MessageScopeStringMapping =
 	{
@@ -33,10 +33,10 @@ bool FWebSocketDeserializedMessage::ParseJson(const FString& Json, FString& OutP
 
 	TSharedPtr<FJsonValue> RootValue;
 
-	TSharedRef<TJsonReader<TCHAR>> JsonReader = TJsonReaderFactory<TCHAR>::Create(Json);
+	const TSharedRef<TJsonReader<TCHAR>> JsonReader = TJsonReaderFactory<TCHAR>::Create(InJson);
 	if (FJsonSerializer::Deserialize(JsonReader, RootValue))
 	{
-		TSharedPtr<FJsonObject> RootObject = RootValue->AsObject();
+		const TSharedPtr<FJsonObject> RootObject = RootValue->AsObject();
 		if (!RootObject)
 		{
 			return false;
diff --git a/Engine/Plugins/Experimental/WebSocketMessaging/Source/WebSocketMessaging/Private/WebSocketDeserializedMessage.h b/Engine/Plugins/Experimental/WebSocketMessaging/Source/WebSocketMessaging/Private/WebSocketDeserializedMessage.h
--- a/Engine/Plugins/Experimental/WebSocketMessaging/Source/WebSocketMessaging/Private/WebSocketDeserializedMessage.h
+++ b/Engine/Plugins/Experimental/WebSocketMessaging/Source/WebSocketMessaging/Private/WebSocketDeserializedMessage.h
@@ -18,7 +18,7 @@ public:
 
 	virtual ~FWebSocketDeserializedMessage() override;
 
-	bool ParseJson(const FString& Json, FString& OutParseError);
+	bool ParseJson(const FString& InJson, FString& OutParseError);
 
 	virtual TSharedPtr<IMessageAttachment, ESPMode::ThreadSafe> GetAttachment() const override
 	{
diff --git a/Engine/Plugins/Experimental/WebSocketMessaging/Source/WebSocketMessaging/Private/WebSocketMessageTransport.cpp b/Engine/Plugins/Experimental/WebSocketMessaging/Source/WebSocketMessaging/Private/WebSocketMessageTransport.cpp
--- a/Engine/Plugins/Experimental/WebSocketMessaging/Source/WebSocketMessaging/Private/WebSocketMessageTransport.cpp
+++ b/Engine/Plugins/Experimental/WebSocketMessaging/Source/WebSocketMessaging/Private/WebSocketMessageTransport.cpp
@@ -52,11 +52,11 @@ FName FWebSocketMessageTransport::GetDebugName() const
 	return DebugName;
 }
 
-bool FWebSocketMessageTransport::StartTransport(IMessageTransportHandler& Handler)
+bool FWebSocketMessageTransport::StartTransport(IMessageTransportHandler& InHandler)
 {
 	const UWebSocketMessagingSettings* Settings = GetDefault<UWebSocketMessagingSettings>();
 	
-	TransportHandler = &Handler;
+	TransportHandler = &InHandler;
 	
 	const int32 ServerPort = Settings->GetServerPort();
 	
@@ -163,43 +163,43 @@ bool FWebSocketMessageTransport::NeedsRestart() const
 	return false;
 }
 
-void FWebSocketMessageTransport::OnClosed(int32 Code, const FString& Reason, bool bUserClose, FWebSocketMessageConnectionRef WebSocketMessageConnection)
+void FWebSocketMessageTransport::OnClosed(int32 InCode, const FString& InReason, bool bInUserClose, FWebSocketMessageConnectionRef InWebSocketMessageConnection)
 {
-	UE_LOG(LogWebSocketMessaging, Log, TEXT("Connection to %s closed, Code: %d Reason: \"%s\" UserClose: %s, retrying..."), *WebSocketMessageConnection->Url, Code, *Reason, bUserClose ? TEXT("true") : TEXT("false"));
-	ForgetTransportNode(WebSocketMessageConnection);
-	WebSocketMessageConnection->bIsConnecting = false;
-	RetryConnection(WebSocketMessageConnection);
+	UE_LOG(LogWebSocketMessaging, Log, TEXT("Connection to %s closed, Code: %d Reason: \"%s\" UserClose: %s, retrying..."), *InWebSocketMessageConnection->Url, InCode, *InReason, bInUserClose ? TEXT("true") : TEXT("false"));
+	ForgetTransportNode(InWebSocketMessageConnection);
+	InWebSocketMessageConnection->bIsConnecting = false;
+	RetryConnection(InWebSocketMessageConnection);
 }
 
-void FWebSocketMessageTransport::OnConnectionError(const FString& Message, FWebSocketMessageConnectionRef WebSocketMessageConnection)
+void FWebSocketMessageTransport::OnConnectionError(const FString& InMessage, FWebSocketMessageConnectionRef InWebSocketMessageConnection)
 {
-	if (!WebSocketMessageConnection->bIsConnecting)
+	if (!InWebSocketMessageConnection->bIsConnecting)
 	{
-		UE_LOG(LogWebSocketMessaging, Log, TEXT("Connection to %s error: %s, retrying..."), *WebSocketMessageConnection->Url, *Message);
+		UE_LOG(LogWebSocketMessaging, Log, TEXT("Connection to %s error: %s, retrying..."), *InWebSocketMessageConnection->Url, *InMessage);
 	}
-	ForgetTransportNode(WebSocketMessageConnection);
-	WebSocketMessageConnection->bIsConnecting = false;
-	RetryConnection(WebSocketMessageConnection);
+	ForgetTransportNode(InWebSocketMessageConnection);
+	InWebSocketMessageConnection->bIsConnecting = false;
+	RetryConnection(InWebSocketMessageConnection);
 }
 
-void FWebSocketMessageTransport::OnJsonMessage(const FString& Message, FWebSocketMessageConnectionRef WebSocketMessageConnection)
+void FWebSocketMessageTransport::OnJsonMessage(const FString& InMessage, FWebSocketMessageConnectionRef InWebSocketMessageConnection)
 {
 	FString ParseError;
 	const TSharedRef<FWebSocketDeserializedMessage> Context = MakeShared<FWebSocketDeserializedMessage>();
-	if (Context->ParseJson(Message, ParseError))
+	if (Context->ParseJson(InMessage, ParseError))
 	{
-		TransportHandler->ReceiveTransportMessage(Context, WebSocketMessageConnection->Guid);
+		TransportHandler->ReceiveTransportMessage(Context, InWebSocketMessageConnection->Guid);
 	}
 	else
 	{
-		UE_LOG(LogWebSocketMessaging, Log, TEXT("Invalid Json Message received on %s: %s"), *WebSocketMessageConnection->Url, *ParseError);
+		UE_LOG(LogWebSocketMessaging, Log, TEXT("Invalid Json Message received on %s: %s"), *InWebSocketMessageConnection->Url, *ParseError);
 	}
 }
 
-void FWebSocketMessageTransport::OnServerJsonMessage(void* Data, int32 DataSize, FWebSocketMessageConnectionRef WebSocketMessageConnection)
+void FWebSocketMessageTransport::OnServerJsonMessage(void* InData, int32 InDataSize, FWebSocketMessageConnectionRef InWebSocketMessageConnection)
 {
-	FString Message(DataSize, reinterpret_cast<UTF8CHAR*>(Data));
-	OnJsonMessage(Message, WebSocketMessageConnection);
+	FString Message(InDataSize, reinterpret_cast<UTF8CHAR*>(InData));
+	OnJsonMessage(Message, InWebSocketMessageConnection);
 }
 
 class FWebSocketMessageTransportSerializeHelper
@@ -327,11 +327,11 @@ const TMap<EMessageScope, FString> FWebSocketMessageTransportSerializeHelper::Me
 	{EMessageScope::All, "All"}
 };
 
-bool FWebSocketMessageTransport::TransportMessage(const TSharedRef<IMessageContext, ESPMode::ThreadSafe>& Context, const TArray<FGuid>& Recipients)
+bool FWebSocketMessageTransport::TransportMessage(const TSharedRef<IMessageContext, ESPMode::ThreadSafe>& InContext, const TArray<FGuid>& InRecipients)
 {
 	TMap<FGuid, FWebSocketMessageConnectionRef> RecipientConnections;
 
-	if (Recipients.Num() == 0)
+	if (InRecipients.Num() == 0)
 	{
 		// broadcast the message to all valid connections
 		RecipientConnections = WebSocketMessageConnections.FilterByPredicate([](const TPair<FGuid, FWebSocketMessageConnectionRef>& Pair) -> bool
@@ -342,7 +342,7 @@ bool FWebSocketMessageTransport::TransportMessage(const TSharedRef<IMessageConte
 	else
 	{
 		// Find connections for each recipient.  We do not transport unicast messages for unknown nodes.
-		for (const FGuid& Recipient : Recipients)
+		for (const FGuid& Recipient : InRecipients)
 		{
 			FWebSocketMessageConnectionRef* RecipientConnection = WebSocketMessageConnections.Find(Recipient);
 			if (RecipientConnection && !(*RecipientConnection)->bDestroyed && (*RecipientConnection)->IsConnected())
@@ -370,7 +370,7 @@ bool FWebSocketMessageTransport::TransportMessage(const TSharedRef<IMessageConte
 		if (Connection.Value->WebSocketConnection.IsValid())
 		{
 			// Remark: client connections are always text/json
-			if (JsonSerializer.SerializeOnDemand(Context))
+			if (JsonSerializer.SerializeOnDemand(InContext))
 			{
 				Connection.Value->WebSocketConnection->Send(JsonSerializer.OutputMessage);
 			}
@@ -380,7 +380,7 @@ bool FWebSocketMessageTransport::TransportMessage(const TSharedRef<IMessageConte
 			// Remark: server connections are always binary.
 			if (Settings->ServerTransportFormat == EWebSocketMessagingTransportFormat::Json)
 			{
-				if (JsonSerializer.SerializeOnDemand(Context))
+				if (JsonSerializer.SerializeOnDemand(InContext))
 				{
 					auto MessageUtf8 = StringCast<UTF8CHAR>(*JsonSerializer.OutputMessage);
 					Connection.Value->WebSocketServerConnection->Send(
@@ -389,7 +389,7 @@ bool FWebSocketMessageTransport::TransportMessage(const TSharedRef<IMessageConte
 			}
 			else
 			{
-				if (CborSerializer.SerializeOnDemand(Context))
+				if (CborSerializer.SerializeOnDemand(InContext))
 				{
 					Connection.Value->WebSocketServerConnection->Send(
 						CborSerializer.OutputMessage.GetData(), CborSerializer.OutputMessage.Num(), /*bPrependSize*/ false);
@@ -401,27 +401,27 @@ bool FWebSocketMessageTransport::TransportMessage(const TSharedRef<IMessageConte
 	return true;
 }
 
-void FWebSocketMessageTransport::OnConnected(FWebSocketMessageConnectionRef WebSocketMessageConnection)
+void FWebSocketMessageTransport::OnConnected(FWebSocketMessageConnectionRef InWebSocketMessageConnection)
 {
-	UE_LOG(LogWebSocketMessaging, Log, TEXT("Connected to %s"), *WebSocketMessageConnection->Url);
-	WebSocketMessageConnection->bIsConnecting = false;
+	UE_LOG(LogWebSocketMessaging, Log, TEXT("Connected to %s"), *InWebSocketMessageConnection->Url);
+	InWebSocketMessageConnection->bIsConnecting = false;
 }
 
-void FWebSocketMessageTransport::OnServerConnectionClosed(FWebSocketMessageConnectionRef WebSocketMessageConnection)
+void FWebSocketMessageTransport::OnServerConnectionClosed(FWebSocketMessageConnectionRef InWebSocketMessageConnection)
 {
-	UE_LOG(LogWebSocketMessaging, Log, TEXT("%s disconnected"), *WebSocketMessageConnection->Url);
-	ForgetTransportNode(WebSocketMessageConnection);
-	WebSocketMessageConnections.Remove(WebSocketMessageConnection->Guid);
+	UE_LOG(LogWebSocketMessaging, Log, TEXT("%s disconnected"), *InWebSocketMessageConnection->Url);
+	ForgetTransportNode(InWebSocketMessageConnection);
+	WebSocketMessageConnections.Remove(InWebSocketMessageConnection->Guid);
 }
 
-void FWebSocketMessageTransport::RetryConnection(FWebSocketMessageConnectionRef WebSocketMessageConnection)
+void FWebSocketMessageTransport::RetryConnection(FWebSocketMessageConnectionRef InWebSocketMessageConnection)
 {
-	if (WebSocketMessageConnection->bIsConnecting)
+	if (InWebSocketMessageConnection->bIsConnecting)
 	{
 		return;
 	}
 
-	WebSocketMessageConnection->RetryHandle = FTSTicker::GetCoreTicker().AddTicker(FTickerDelegate::CreateLambda([](float DeltaTime, FWebSocketMessageConnectionRef WebSocketMessageConnection)
+	InWebSocketMessageConnection->RetryHandle = FTSTicker::GetCoreTicker().AddTicker(FTickerDelegate::CreateLambda([](float DeltaTime, FWebSocketMessageConnectionRef WebSocketMessageConnection)
 		{
 			if (!WebSocketMessageConnection->bDestroyed && !WebSocketMessageConnection->bIsConnecting && !WebSocketMessageConnection->WebSocketConnection->IsConnected())
 			{
@@ -429,26 +429,26 @@ void FWebSocketMessageTransport::RetryConnection(FWebSocketMessageConnectionRef
 				WebSocketMessageConnection->WebSocketConnection->Connect();
 			}
 			return false;
-		}, WebSocketMessageConnection), 1.0f);
+		}, InWebSocketMessageConnection), 1.0f);
 }
 
-void FWebSocketMessageTransport::ClientConnected(INetworkingWebSocket* NetworkingWebSocket)
+void FWebSocketMessageTransport::ClientConnected(INetworkingWebSocket* InNetworkingWebSocket)
 {
-	FString RemoteEndPoint = NetworkingWebSocket->RemoteEndPoint(true);
+	FString RemoteEndPoint = InNetworkingWebSocket->RemoteEndPoint(true);
 	UE_LOG(LogWebSocketMessaging, Log, TEXT("New WebSocket Server connection: %s"), *RemoteEndPoint);
 
 	FGuid Guid = FGuid::NewGuid();
 
-	FWebSocketMessageConnectionRef WebSocketMessageConnection = MakeShared<FWebSocketMessageConnection>(RemoteEndPoint, Guid, NetworkingWebSocket);
+	FWebSocketMessageConnectionRef WebSocketMessageConnection = MakeShared<FWebSocketMessageConnection>(RemoteEndPoint, Guid, InNetworkingWebSocket);
 
-	NetworkingWebSocket->SetReceiveCallBack(FWebSocketPacketReceivedCallBack::CreateThreadSafeSP(this, &FWebSocketMessageTransport::OnServerJsonMessage, WebSocketMessageConnection));
-	NetworkingWebSocket->SetSocketClosedCallBack(FWebSocketInfoCallBack::CreateThreadSafeSP(this, &FWebSocketMessageTransport::OnServerConnectionClosed, WebSocketMessageConnection));
-	NetworkingWebSocket->SetErrorCallBack(FWebSocketInfoCallBack::CreateThreadSafeSP(this, &FWebSocketMessageTransport::OnServerConnectionClosed, WebSocketMessageConnection));
+	InNetworkingWebSocket->SetReceiveCallBack(FWebSocketPacketReceivedCallBack::CreateThreadSafeSP(this, &FWebSocketMessageTransport::OnServerJsonMessage, WebSocketMessageConnection));
+	InNetworkingWebSocket->SetSocketClosedCallBack(FWebSocketInfoCallBack::CreateThreadSafeSP(this, &FWebSocketMessageTransport::OnServerConnectionClosed, WebSocketMessageConnection));
+	InNetworkingWebSocket->SetErrorCallBack(FWebSocketInfoCallBack::CreateThreadSafeSP(this, &FWebSocketMessageTransport::OnServerConnectionClosed, WebSocketMessageConnection));
 
 	WebSocketMessageConnections.Add(Guid, WebSocketMessageConnection);
 }
 
-bool FWebSocketMessageTransport::ServerTick(float DeltaTime)
+bool FWebSocketMessageTransport::ServerTick(float InDeltaTime)
 {
 	if (Server.IsValid())
 	{
@@ -458,10 +458,10 @@ bool FWebSocketMessageTransport::ServerTick(float DeltaTime)
 	return true;
 }
 
-void FWebSocketMessageTransport::ForgetTransportNode(FWebSocketMessageConnectionRef WebSocketMessageConnection)
+void FWebSocketMessageTransport::ForgetTransportNode(FWebSocketMessageConnectionRef InWebSocketMessageConnection)
 {
 	if (TransportHandler)
 	{
-		TransportHandler->ForgetTransportNode(WebSocketMessageConnection->Guid);
+		TransportHandler->ForgetTransportNode(InWebSocketMessageConnection->Guid);
 	}
 }
diff --git a/Engine/Plugins/Experimental/WebSocketMessaging/Source/WebSocketMessaging/Private/WebSocketMessageTransport.h b/Engine/Plugins/Experimental/WebSocketMessaging/Source/WebSocketMessaging/Private/WebSocketMessageTransport.h
--- a/Engine/Plugins/Experimental/WebSocketMessaging/Source/WebSocketMessaging/Private/WebSocketMessageTransport.h
+++ b/Engine/Plugins/Experimental/WebSocketMessaging/Source/WebSocketMessaging/Private/WebSocketMessageTransport.h
@@ -11,7 +11,7 @@
 struct FWebSocketMessageConnection
 {
 	FWebSocketMessageConnection() = delete;
-	FWebSocketMessageConnection(const FString& InUrl, FGuid InGuid, TSharedRef<class IWebSocket, ESPMode::ThreadSafe> InWebSocketConnection) :
+	FWebSocketMessageConnection(const FString& InUrl, const FGuid& InGuid, const TSharedRef<class IWebSocket, ESPMode::ThreadSafe>& InWebSocketConnection) :
 		Url(InUrl),
 		Guid(InGuid),
 		WebSocketConnection(InWebSocketConnection),
@@ -22,7 +22,7 @@ struct FWebSocketMessageConnection
 
 	}
 
-	FWebSocketMessageConnection(const FString& InUrl, FGuid InGuid, class INetworkingWebSocket* InWebSocketServerConnection) :
+	FWebSocketMessageConnection(const FString& InUrl, const FGuid& InGuid, class INetworkingWebSocket* InWebSocketServerConnection) :
 		Url(InUrl),
 		Guid(InGuid),
 		WebSocketConnection(nullptr),
@@ -69,29 +69,29 @@ public:
 
 	//~ Begin IMessageTransport
 	virtual FName GetDebugName() const override;
-	virtual bool StartTransport(IMessageTransportHandler& Handler) override;
+	virtual bool StartTransport(IMessageTransportHandler& InHandler) override;
 	virtual void StopTransport() override;
-	virtual bool TransportMessage(const TSharedRef<IMessageContext, ESPMode::ThreadSafe>& Context, const TArray<FGuid>& Recipients) override;
+	virtual bool TransportMessage(const TSharedRef<IMessageContext, ESPMode::ThreadSafe>& InContext, const TArray<FGuid>& InRecipients) override;
 	//~ End IMessageTransport
 
 	bool NeedsRestart() const;
 
-	virtual void OnJsonMessage(const FString& Message, FWebSocketMessageConnectionRef WebSocketMessageConnection);
-	virtual void OnConnected(FWebSocketMessageConnectionRef WebSocketMessageConnection);
-	virtual void OnClosed(int32 Code, const FString& Reason, bool bUserClose, FWebSocketMessageConnectionRef WebSocketMessageConnection);
-	virtual void OnConnectionError(const FString& Message, FWebSocketMessageConnectionRef WebSocketMessageConnection);
+	virtual void OnJsonMessage(const FString& InMessage, FWebSocketMessageConnectionRef InWebSocketMessageConnection);
+	virtual void OnConnected(FWebSocketMessageConnectionRef InWebSocketMessageConnection);
+	virtual void OnClosed(int32 InCode, const FString& InReason, bool bInUserClose, FWebSocketMessageConnectionRef InWebSocketMessageConnection);
+	virtual void OnConnectionError(const FString& InMessage, FWebSocketMessageConnectionRef InWebSocketMessageConnection);
 
-	virtual void RetryConnection(FWebSocketMessageConnectionRef WebSocketMessageConnection);
+	virtual void RetryConnection(FWebSocketMessageConnectionRef InWebSocketMessageConnection);
 
-	virtual void ClientConnected(class INetworkingWebSocket* NetworkingWebSocket);
+	virtual void ClientConnected(class INetworkingWebSocket* InNetworkingWebSocket);
 
-	virtual bool ServerTick(float DeltaTime);
+	virtual bool ServerTick(float InDeltaTime);
 
-	virtual void OnServerJsonMessage(void* Data, int32 DataSize, FWebSocketMessageConnectionRef WebSocketMessageConnection);
-	virtual void OnServerConnectionClosed(FWebSocketMessageConnectionRef WebSocketMessageConnection);
+	virtual void OnServerJsonMessage(void* InData, int32 InDataSize, FWebSocketMessageConnectionRef InWebSocketMessageConnection);
+	virtual void OnServerConnectionClosed(FWebSocketMessageConnectionRef InWebSocketMessageConnection);
 
 protected:
-	void ForgetTransportNode(FWebSocketMessageConnectionRef WebSocketMessageConnection);
+	void ForgetTransportNode(FWebSocketMessageConnectionRef InWebSocketMessageConnection);
 	
 	IMessageTransportHandler* TransportHandler = nullptr;
 	TMap<FGuid, FWebSocketMessageConnectionRef> WebSocketMessageConnections;
