[Websocket Messaging] Adding a description with the word "websocket" (in one word) so the plugin is easier to find for users.

diff --git a/Engine/Plugins/Experimental/WebSocketMessaging/WebSocketMessaging.uplugin b/Engine/Plugins/Experimental/WebSocketMessaging/WebSocketMessaging.uplugin
--- a/Engine/Plugins/Experimental/WebSocketMessaging/WebSocketMessaging.uplugin
+++ b/Engine/Plugins/Experimental/WebSocketMessaging/WebSocketMessaging.uplugin
@@ -3,7 +3,7 @@
 	"Version": 1,
 	"VersionName": "1.0",
 	"FriendlyName": "Web Socket Messaging",
-	"Description": "",
+	"Description": "Adds a WebSocket based transport layer to the messaging sub-system for sending and receiving messages between networked computers and devices.",
 	"Category": "Experimental",
 	"CreatedBy": "Epic Games, Inc.",
 	"CreatedByURL": "https://epicgames.com",
