[WebSocketMessaging] Changing default Server Transport format to JSON.

diff --git a/Engine/Plugins/Experimental/WebSocketMessaging/Source/WebSocketMessaging/Private/WebSocketMessagingSettings.h b/Engine/Plugins/Experimental/WebSocketMessaging/Source/WebSocketMessaging/Private/WebSocketMessagingSettings.h
--- a/Engine/Plugins/Experimental/WebSocketMessaging/Source/WebSocketMessaging/Private/WebSocketMessagingSettings.h
+++ b/Engine/Plugins/Experimental/WebSocketMessaging/Source/WebSocketMessaging/Private/WebSocketMessagingSettings.h
@@ -44,7 +44,7 @@ public:
 
 	/** Format used to serialize the messages on the server's WebSockets.*/
 	UPROPERTY(Config, EditAnywhere, Category=Transport)
-	EWebSocketMessagingTransportFormat ServerTransportFormat = EWebSocketMessagingTransportFormat::Cbor;
+	EWebSocketMessagingTransportFormat ServerTransportFormat = EWebSocketMessagingTransportFormat::Json;
 
 	/**
 	 * For Json formatting only:
